Adding the ability for the Media Player to playback WebRTC streams; including those from another Pixel Streaming session.
- Currently only VP8 and VP9 encoded streams are supported (more to come).

diff --git a/Engine/Plugins/Media/PixelStreamingPlayer/PixelStreamingPlayer.uplugin b/Engine/Plugins/Media/PixelStreamingPlayer/PixelStreamingPlayer.uplugin
new file mode 100644
--- /dev/null
+++ b/Engine/Plugins/Media/PixelStreamingPlayer/PixelStreamingPlayer.uplugin
@@ -0,0 +1,39 @@
+{
+	"FileVersion": 3,
+	"Version": 1.0,
+	"VersionName": "1.0",
+	"FriendlyName": "Pixel Streaming Player",
+	"Description": "Support for receiving a pixel streaming stream and displaying it in game.",
+	"Category": "Graphics",
+	"CreatedBy": "Epic Games, Inc.",
+	"CreatedByURL": "http://epicgames.com",
+	"DocsURL": "https://docs.unrealengine.com/en-US/Platforms/PixelStreaming/index.html",
+	"MarketplaceURL": "",
+	"SupportURL": "",
+	"EnabledByDefault": false,
+	"CanContainContent": true,
+	"IsBetaVersion": true,
+	"Installed": false,
+	"Modules": [
+		{
+			"Name": "PixelStreamingPlayer",
+			"Type": "Runtime",
+			"LoadingPhase": "PreDefault",
+			"PlatformAllowList": [ "Win64", "Linux" ],
+			"TargetDenyList": [ "Server" ]
+		},
+		{
+			"Name": "PixelStreamingPlayerEditor",
+			"Type": "Editor",
+			"LoadingPhase": "PostEngineInit",
+			"PlatformAllowList": [ "Win64", "Linux" ],
+			"TargetDenyList": [ "Server" ]
+		}
+	],
+	"Plugins": [
+		{
+			"Name": "PixelStreaming",
+			"Enabled": true
+		}
+	]
+}
diff --git a/Engine/Plugins/Media/PixelStreamingPlayer/Source/PixelStreamingPlayer/PixelStreamingPlayer.Build.cs b/Engine/Plugins/Media/PixelStreamingPlayer/Source/PixelStreamingPlayer/PixelStreamingPlayer.Build.cs
new file mode 100644
--- /dev/null
+++ b/Engine/Plugins/Media/PixelStreamingPlayer/Source/PixelStreamingPlayer/PixelStreamingPlayer.Build.cs
@@ -0,0 +1,46 @@
+// Copyright Epic Games, Inc. All Rights Reserved.
+
+using System;
+using System.IO;
+using System.Collections.Generic;
+using EpicGames.Core;
+
+namespace UnrealBuildTool.Rules
+{
+	public class PixelStreamingPlayer : ModuleRules
+	{
+		public PixelStreamingPlayer(ReadOnlyTargetRules Target) : base(Target)
+		{
+			PCHUsage = PCHUsageMode.UseExplicitOrSharedPCHs;
+			PrivatePCHHeaderFile = "Private/PCH.h";
+
+			var EngineDir = Path.GetFullPath(Target.RelativeEnginePath);
+
+			PublicIncludePaths.AddRange(
+				new string[] {
+				});
+
+			PrivateIncludePaths.AddRange(
+				new string[] {
+					Path.Combine(EngineDir, "Plugins/Media/PixelStreaming/Source/PixelStreaming/Private")
+				});
+
+			PublicDependencyModuleNames.AddRange(
+				new string[] {
+					"WebRTC",
+				});
+
+			PrivateDependencyModuleNames.AddRange(
+				new string[] {
+					"Core",
+					"CoreUObject",
+					"Engine",
+					"MediaAssets",
+					"PixelStreaming",
+					"RenderCore",
+					"RHI",
+					"WebSockets",
+				});
+		}
+	}
+}
diff --git a/Engine/Plugins/Media/PixelStreamingPlayer/Source/PixelStreamingPlayer/Private/PCH.h b/Engine/Plugins/Media/PixelStreamingPlayer/Source/PixelStreamingPlayer/Private/PCH.h
new file mode 100644
--- /dev/null
+++ b/Engine/Plugins/Media/PixelStreamingPlayer/Source/PixelStreamingPlayer/Private/PCH.h
@@ -0,0 +1,7 @@
+// Copyright Epic Games, Inc. All Rights Reserved.
+
+#pragma once
+
+#include "WebSocketsModule.h"
+#include "WebRTCIncludes.h"
+#include "PixelStreamingPlayerPrivate.h"
diff --git a/Engine/Plugins/Media/PixelStreamingPlayer/Source/PixelStreamingPlayer/Private/PixelStreamingMediaTexture.cpp b/Engine/Plugins/Media/PixelStreamingPlayer/Source/PixelStreamingPlayer/Private/PixelStreamingMediaTexture.cpp
new file mode 100644
--- /dev/null
+++ b/Engine/Plugins/Media/PixelStreamingPlayer/Source/PixelStreamingPlayer/Private/PixelStreamingMediaTexture.cpp
@@ -0,0 +1,235 @@
+// Copyright Epic Games, Inc. All Rights Reserved.
+
+#include "PixelStreamingMediaTexture.h"
+#include "PixelStreamingMediaTextureResource.h"
+#include "Async/Async.h"
+#include "WebRTCIncludes.h"
+
+constexpr int32 DEFAULT_WIDTH = 1920;
+constexpr int32 DEFAULT_HEIGHT = 1080;
+
+// This is copy-pasted from Engine\Source\Runtime\RHI\Private\RHIUtilities.cpp, TODO: RHICreateTargetableShaderResource family of functions is deprecated,
+// suggesting RHICreateTexture be used instead, so translate to that usage of the API at some point
+void RHICreateTargetableShaderResource(
+	const FRHITextureCreateDesc& BaseDesc,
+	ETextureCreateFlags TargetableTextureFlags,
+	bool bForceSeparateTargetAndShaderResource,
+	bool bForceSharedTargetAndShaderResource,
+	FTextureRHIRef& OutTargetableTexture,
+	FTextureRHIRef& OutShaderResourceTexture
+	)
+{
+	// Ensure none of the usage flags are passed in.
+	check(!EnumHasAnyFlags(BaseDesc.Flags, ETextureCreateFlags::RenderTargetable | ETextureCreateFlags::ResolveTargetable));
+
+	// Ensure we aren't forcing separate and shared textures at the same time.
+	check(!(bForceSeparateTargetAndShaderResource && bForceSharedTargetAndShaderResource));
+
+	// Ensure that the targetable texture is either render or depth-stencil targetable.
+	check(EnumHasAnyFlags(TargetableTextureFlags, ETextureCreateFlags::RenderTargetable | ETextureCreateFlags::DepthStencilTargetable | ETextureCreateFlags::UAV));
+
+	if (BaseDesc.NumSamples > 1 && !bForceSharedTargetAndShaderResource)
+	{
+		bForceSeparateTargetAndShaderResource = RHISupportsSeparateMSAAAndResolveTextures(GMaxRHIShaderPlatform);
+	}
+
+	if (!bForceSeparateTargetAndShaderResource)
+	{
+		FRHITextureCreateDesc Desc =
+			FRHITextureCreateDesc(BaseDesc)
+			.AddFlags(TargetableTextureFlags | ETextureCreateFlags::ShaderResource)
+			.SetInitialState(ERHIAccess::SRVMask);
+
+		// Create a single texture that has both TargetableTextureFlags and ETextureCreateFlags::ShaderResource set.
+		OutTargetableTexture = OutShaderResourceTexture = RHICreateTexture(Desc);
+	}
+	else
+	{
+		ETextureCreateFlags ResolveTargetableTextureFlags = ETextureCreateFlags::ResolveTargetable;
+		if (EnumHasAnyFlags(TargetableTextureFlags, ETextureCreateFlags::DepthStencilTargetable))
+		{
+			ResolveTargetableTextureFlags |= ETextureCreateFlags::DepthStencilResolveTarget;
+		}
+
+		FRHITextureCreateDesc TargetableDesc =
+			FRHITextureCreateDesc(BaseDesc)
+			.AddFlags(TargetableTextureFlags)
+			.SetInitialState(ERHIAccess::SRVMask);
+
+		FRHITextureCreateDesc ResourceDesc =
+			FRHITextureCreateDesc(BaseDesc)
+			.AddFlags(ResolveTargetableTextureFlags | ETextureCreateFlags::ShaderResource)
+			.SetInitialState(ERHIAccess::SRVMask)
+			.SetNumSamples(1);
+
+		// Create a texture that has TargetableTextureFlags set, and a second texture that has ETextureCreateFlags::ResolveTargetable and ETextureCreateFlags::ShaderResource set.
+		OutTargetableTexture = RHICreateTexture(TargetableDesc);
+		OutShaderResourceTexture = RHICreateTexture(ResourceDesc);
+	}
+}
+
+UPixelStreamingMediaTexture::UPixelStreamingMediaTexture(const FObjectInitializer& ObjectInitializer)
+	: Super(ObjectInitializer)
+{
+	SetResource(nullptr);
+}
+
+void UPixelStreamingMediaTexture::BeginDestroy()
+{
+	AsyncTask(ENamedThreads::ActualRenderingThread, [&RenderSyncContext = RenderSyncContext, &RenderTarget = RenderTarget]() {
+		FScopeLock Lock(&RenderSyncContext);
+		RenderTarget = nullptr;
+	});
+
+	Super::BeginDestroy();
+}
+
+float UPixelStreamingMediaTexture::GetSurfaceHeight() const
+{
+	return CurrentResource != nullptr ? CurrentResource->GetSizeY() : 0.0f;
+}
+
+float UPixelStreamingMediaTexture::GetSurfaceWidth() const
+{
+	return CurrentResource != nullptr ? CurrentResource->GetSizeX() : 0.0f;
+}
+
+void UPixelStreamingMediaTexture::GetResourceSizeEx(FResourceSizeEx& CumulativeResourceSize)
+{
+	Super::GetResourceSizeEx(CumulativeResourceSize);
+
+	if (CurrentResource != nullptr)
+	{
+		CumulativeResourceSize.AddUnknownMemoryBytes(CurrentResource->GetResourceSize());
+	}
+}
+
+EMaterialValueType UPixelStreamingMediaTexture::GetMaterialType() const
+{
+	return MCT_Texture2D;
+}
+
+void UPixelStreamingMediaTexture::OnFrame(const webrtc::VideoFrame& frame)
+{
+	// TODO
+	// Currently this will cause all input frames to be converted to I420 (using webrtc::ConvertFromI420)
+	// and then render the rgba buffer to a texture, but in some cases we might already have the native
+	// texture and can use it or copy it directly. At some point it would be useful to detect and fork this
+	// behaviour.
+	// Additionally, we're rendering to a texture then updating the texture resource with every frame. We
+	// might be able to just map a single dest texture and update it each frame.
+
+	constexpr auto TEXTURE_PIXEL_FORMAT = PF_B8G8R8A8;
+	constexpr auto WEBRTC_PIXEL_FORMAT = webrtc::VideoType::kARGB;
+
+	AsyncTask(ENamedThreads::ActualRenderingThread, [=]() {
+		FScopeLock Lock(&RenderSyncContext);
+
+		uint32_t Size = webrtc::CalcBufferSize(WEBRTC_PIXEL_FORMAT, frame.width(), frame.height());
+
+		if (Size > BufferSize)
+		{
+			delete[] Buffer;
+			Buffer = new uint8_t[Size];
+			BufferSize = Size;
+		}
+
+		webrtc::ConvertFromI420(frame, WEBRTC_PIXEL_FORMAT, 0, Buffer);
+
+		FIntPoint FrameSize = FIntPoint(frame.width(), frame.height());
+		FRHICommandListImmediate& RHICmdList = FRHICommandListExecutor::GetImmediateCommandList();
+
+		if (!RenderTargetDescriptor.IsValid() || RenderTargetDescriptor.GetSize() != FIntVector(FrameSize.X, FrameSize.Y, 0))
+		{
+			// Create the RenderTarget descriptor
+			RenderTargetDescriptor = FPooledRenderTargetDesc::Create2DDesc(FrameSize,
+				TEXTURE_PIXEL_FORMAT,
+				FClearValueBinding::None,
+				TexCreate_None,
+				TexCreate_RenderTargetable,
+				false);
+
+			// Update the shader resource for the 'SourceTexture'
+			FRHITextureCreateDesc RenderTargetTextureDesc =
+				FRHITextureCreateDesc::Create2D(TEXT(""), FrameSize.X, FrameSize.Y, TEXTURE_PIXEL_FORMAT)
+				.SetClearValue(FClearValueBinding(FLinearColor(0.0f, 0.0f, 0.0f)))
+				.SetFlags(ETextureCreateFlags::Dynamic)
+				.DetermineInititialState();
+
+			FTextureRHIRef DummyTexture2DRHI;
+			RHICreateTargetableShaderResource(RenderTargetTextureDesc, TexCreate_RenderTargetable, false, false, SourceTexture, DummyTexture2DRHI);
+
+			// Find a free target-able texture from the render pool
+			GRenderTargetPool.FindFreeElement(RHICmdList,
+				RenderTargetDescriptor,
+				RenderTarget,
+				TEXT("PIXELSTEAMINGPLAYER"));
+		}
+
+		// Create the update region structure
+		FUpdateTextureRegion2D Region(0, 0, 0, 0, FrameSize.X, FrameSize.Y);
+
+		// Set the Pixel data of the webrtc Frame to the SourceTexture
+		RHIUpdateTexture2D(SourceTexture, 0, Region, frame.width() * 4, (uint8*&)Buffer);
+
+		UpdateTextureReference(RHICmdList, SourceTexture);
+	});
+}
+
+void UPixelStreamingMediaTexture::UpdateTextureReference(FRHICommandList& RHICmdList, FTexture2DRHIRef Reference)
+{
+	if (CurrentResource != nullptr)
+	{
+		if (Reference.IsValid() && CurrentResource->TextureRHI != Reference)
+		{
+			CurrentResource->TextureRHI = Reference;
+			RHIUpdateTextureReference(TextureReference.TextureReferenceRHI, CurrentResource->TextureRHI);
+		}
+		else if (!Reference.IsValid())
+		{
+			if (CurrentResource != nullptr)
+			{
+				InitializeResources();
+
+				// Make sure RenderThread is executed before continuing
+				FlushRenderingCommands();
+			}
+		}
+	}
+}
+
+FTextureResource* UPixelStreamingMediaTexture::CreateResource()
+{
+	if (CurrentResource != nullptr)
+	{
+		SetResource(nullptr);
+		delete CurrentResource;
+	}
+
+	CurrentResource = new FPixelStreamingMediaTextureResource(this);
+	InitializeResources();
+
+	return CurrentResource;
+}
+
+void UPixelStreamingMediaTexture::InitializeResources()
+{
+	// Set the default video texture to reference nothing
+	FTextureRHIRef ShaderTexture2D;
+	FTextureRHIRef RenderableTexture;
+
+	FRHITextureCreateDesc RenderTargetTextureDesc =
+		FRHITextureCreateDesc::Create2D(TEXT(""), DEFAULT_WIDTH, DEFAULT_HEIGHT, EPixelFormat::PF_B8G8R8A8)
+		.SetClearValue(FClearValueBinding(FLinearColor(0.0f, 0.0f, 0.0f)))
+		.SetFlags(ETextureCreateFlags::Dynamic)
+		.DetermineInititialState();
+
+	RHICreateTargetableShaderResource(RenderTargetTextureDesc, TexCreate_RenderTargetable, false, false, RenderableTexture, ShaderTexture2D);
+
+	CurrentResource->TextureRHI = RenderableTexture;
+
+	ENQUEUE_RENDER_COMMAND(FPixelStreamingMediaTextureUpdateTextureReference)
+	([this](FRHICommandListImmediate& RHICmdList) {
+		RHIUpdateTextureReference(TextureReference.TextureReferenceRHI, CurrentResource->TextureRHI);
+	});
+}
diff --git a/Engine/Plugins/Media/PixelStreamingPlayer/Source/PixelStreamingPlayer/Private/PixelStreamingMediaTextureResource.cpp b/Engine/Plugins/Media/PixelStreamingPlayer/Source/PixelStreamingPlayer/Private/PixelStreamingMediaTextureResource.cpp
new file mode 100644
--- /dev/null
+++ b/Engine/Plugins/Media/PixelStreamingPlayer/Source/PixelStreamingPlayer/Private/PixelStreamingMediaTextureResource.cpp
@@ -0,0 +1,46 @@
+// Copyright Epic Games, Inc. All Rights Reserved.
+
+#include "PixelStreamingMediaTextureResource.h"
+#include "DeviceProfiles/DeviceProfile.h"
+#include "DeviceProfiles/DeviceProfileManager.h"
+
+FPixelStreamingMediaTextureResource::FPixelStreamingMediaTextureResource(UPixelStreamingMediaTexture* Owner)
+{
+	MediaTexture = Owner;
+}
+
+void FPixelStreamingMediaTextureResource::InitDynamicRHI()
+{
+	if (MediaTexture != nullptr)
+	{
+		FSamplerStateInitializerRHI SamplerStateInitializer(
+			(ESamplerFilter)UDeviceProfileManager::Get().GetActiveProfile()->GetTextureLODSettings()->GetSamplerFilter(MediaTexture),
+			AM_Border, AM_Border, AM_Wrap);
+		SamplerStateRHI = RHICreateSamplerState(SamplerStateInitializer);
+	}
+}
+
+void FPixelStreamingMediaTextureResource::ReleaseDynamicRHI()
+{
+	TextureRHI.SafeRelease();
+
+	if (MediaTexture != nullptr)
+	{
+		RHIUpdateTextureReference(MediaTexture->TextureReference.TextureReferenceRHI, nullptr);
+	}
+}
+
+uint32 FPixelStreamingMediaTextureResource::GetSizeX() const
+{
+	return TextureRHI.IsValid() ? TextureRHI->GetSizeXYZ().X : 0;
+}
+
+uint32 FPixelStreamingMediaTextureResource::GetSizeY() const
+{
+	return TextureRHI.IsValid() ? TextureRHI->GetSizeXYZ().Y : 0;
+}
+
+SIZE_T FPixelStreamingMediaTextureResource::GetResourceSize()
+{
+	return CalcTextureSize(GetSizeX(), GetSizeY(), EPixelFormat::PF_A8R8G8B8, 1);
+}
diff --git a/Engine/Plugins/Media/PixelStreamingPlayer/Source/PixelStreamingPlayer/Private/PixelStreamingMediaTextureResource.h b/Engine/Plugins/Media/PixelStreamingPlayer/Source/PixelStreamingPlayer/Private/PixelStreamingMediaTextureResource.h
new file mode 100644
--- /dev/null
+++ b/Engine/Plugins/Media/PixelStreamingPlayer/Source/PixelStreamingPlayer/Private/PixelStreamingMediaTextureResource.h
@@ -0,0 +1,26 @@
+// Copyright Epic Games, Inc. All Rights Reserved.
+
+#pragma once
+
+#include "TextureResource.h"
+#include "PixelStreamingMediaTexture.h"
+
+/**
+ * The actual texture resource for a FPixelStreamingMediaTexture. Contains the RHI Texture and
+ * Sampler information.
+ */
+class PIXELSTREAMINGPLAYER_API FPixelStreamingMediaTextureResource : public FTextureResource
+{
+public:
+	FPixelStreamingMediaTextureResource(UPixelStreamingMediaTexture* Owner);
+
+	virtual void InitDynamicRHI() override;
+	virtual void ReleaseDynamicRHI() override;
+	virtual uint32 GetSizeX() const override;
+	virtual uint32 GetSizeY() const override;
+	
+	SIZE_T GetResourceSize();
+
+private:
+	UPixelStreamingMediaTexture* MediaTexture = nullptr;
+};
diff --git a/Engine/Plugins/Media/PixelStreamingPlayer/Source/PixelStreamingPlayer/Private/PixelStreamingPeerComponent.cpp b/Engine/Plugins/Media/PixelStreamingPlayer/Source/PixelStreamingPlayer/Private/PixelStreamingPeerComponent.cpp
new file mode 100644
--- /dev/null
+++ b/Engine/Plugins/Media/PixelStreamingPlayer/Source/PixelStreamingPlayer/Private/PixelStreamingPeerComponent.cpp
@@ -0,0 +1,43 @@
+// Copyright Epic Games, Inc. All Rights Reserved.
+
+#include "PixelStreamingPeerComponent.h"
+#include "PixelStreamingPlayerPrivate.h"
+
+UPixelStreamingPeerComponent::UPixelStreamingPeerComponent(const FObjectInitializer& ObjectInitializer)
+	: Super(ObjectInitializer)
+{
+}
+
+void UPixelStreamingPeerComponent::Initialize(UPixelStreamingSignallingComponent* InSignallingComponent)
+{
+	SignallingComponent = InSignallingComponent;
+
+	PeerConnection = FPixelStreamingPeerConnection::Create(SignallingComponent->GetConfig());
+	check(PeerConnection);
+	PeerConnection->SetVideoSink(VideoSink);
+	PeerConnection->SetIceCandidateCallback([&SignallingComponent = SignallingComponent](const webrtc::IceCandidateInterface* Candidate) {
+		SignallingComponent->GetConnection()->SendIceCandidate(*Candidate);
+	});
+}
+
+void UPixelStreamingPeerComponent::ReceiveOffer(const FString& Sdp)
+{
+	if (PeerConnection)
+	{
+		PeerConnection->SetSuccessCallback([&SignallingComponent = SignallingComponent](const webrtc::SessionDescriptionInterface* Sdp) {
+			SignallingComponent->GetConnection()->SendAnswer(*Sdp);
+		});
+		PeerConnection->SetFailureCallback([](const FString& ErrorMsg) {
+			UE_LOG(LogPixelStreamingPlayer, Log, TEXT("SetRemoteDescription Failed: %s"), *ErrorMsg);
+		});
+		PeerConnection->SetRemoteDescription(Sdp);
+	}
+}
+
+void UPixelStreamingPeerComponent::ReceiveIceCandidate(const FString& SdpMid, int SdpMLineIndex, const FString& Sdp)
+{
+	if (PeerConnection) 
+	{
+		PeerConnection->AddRemoteIceCandidate(SdpMid, SdpMLineIndex, Sdp);
+	}
+}
diff --git a/Engine/Plugins/Media/PixelStreamingPlayer/Source/PixelStreamingPlayer/Private/PixelStreamingPlayerModule.cpp b/Engine/Plugins/Media/PixelStreamingPlayer/Source/PixelStreamingPlayer/Private/PixelStreamingPlayerModule.cpp
new file mode 100644
--- /dev/null
+++ b/Engine/Plugins/Media/PixelStreamingPlayer/Source/PixelStreamingPlayer/Private/PixelStreamingPlayerModule.cpp
@@ -0,0 +1,23 @@
+// Copyright Epic Games, Inc. All Rights Reserved.
+
+#include "PixelStreamingPlayerPrivate.h"
+#include "Modules/ModuleManager.h"
+#include "WebSocketsModule.h"
+#include "WebRTCIncludes.h"
+
+DEFINE_LOG_CATEGORY(LogPixelStreamingPlayer);
+
+class FPixelStreamingPlayerModule : public IModuleInterface
+{
+public:
+private:
+	/** IModuleInterface implementation */
+	void StartupModule() override
+	{
+		FModuleManager::LoadModuleChecked<FWebSocketsModule>("WebSockets");
+		rtc::InitializeSSL();
+	}
+	void ShutdownModule() override {}
+};
+
+IMPLEMENT_MODULE(FPixelStreamingPlayerModule, PixelStreamingPlayer)
diff --git a/Engine/Plugins/Media/PixelStreamingPlayer/Source/PixelStreamingPlayer/Private/PixelStreamingPlayerPrivate.h b/Engine/Plugins/Media/PixelStreamingPlayer/Source/PixelStreamingPlayer/Private/PixelStreamingPlayerPrivate.h
new file mode 100644
--- /dev/null
+++ b/Engine/Plugins/Media/PixelStreamingPlayer/Source/PixelStreamingPlayer/Private/PixelStreamingPlayerPrivate.h
@@ -0,0 +1,7 @@
+// Copyright Epic Games, Inc. All Rights Reserved.
+
+#pragma once
+
+#include "Logging/LogMacros.h"
+
+DECLARE_LOG_CATEGORY_EXTERN(LogPixelStreamingPlayer, Log, All);
diff --git a/Engine/Plugins/Media/PixelStreamingPlayer/Source/PixelStreamingPlayer/Private/PixelStreamingSignallingComponent.cpp b/Engine/Plugins/Media/PixelStreamingPlayer/Source/PixelStreamingPlayer/Private/PixelStreamingSignallingComponent.cpp
new file mode 100644
--- /dev/null
+++ b/Engine/Plugins/Media/PixelStreamingPlayer/Source/PixelStreamingPlayer/Private/PixelStreamingSignallingComponent.cpp
@@ -0,0 +1,75 @@
+// Copyright Epic Games, Inc. All Rights Reserved.
+
+#include "PixelStreamingSignallingComponent.h"
+#include "WebSocketsModule.h"
+
+UPixelStreamingSignallingComponent::UPixelStreamingSignallingComponent(const FObjectInitializer& ObjectInitializer)
+	: Super(ObjectInitializer)
+{
+	FPixelStreamingSignallingConnection::FWebSocketFactory WebSocketFactory = [](const FString& Url) { return FWebSocketsModule::Get().CreateWebSocket(Url, TEXT("")); };
+	SignallingConnection = MakeUnique<FPixelStreamingSignallingConnection>(WebSocketFactory, *this);
+}
+
+void UPixelStreamingSignallingComponent::Connect(const FString& Url)
+{
+	if (MediaSource == nullptr)
+	{
+		SignallingConnection->Connect(Url);
+	}
+	else
+	{
+		SignallingConnection->Connect(MediaSource->GetUrl());
+	}
+}
+
+void UPixelStreamingSignallingComponent::Disconnect()
+{
+	SignallingConnection->Disconnect();
+}
+
+void UPixelStreamingSignallingComponent::OnSignallingConnected()
+{
+	OnConnected.Broadcast();
+}
+
+void UPixelStreamingSignallingComponent::OnSignallingDisconnected(int32 StatusCode, const FString& Reason, bool bWasClean)
+{
+	OnDisconnected.Broadcast(StatusCode, Reason, bWasClean);
+}
+
+void UPixelStreamingSignallingComponent::OnSignallingError(const FString& ErrorMsg)
+{
+	OnConnectionError.Broadcast(ErrorMsg);
+}
+
+void UPixelStreamingSignallingComponent::OnSignallingConfig(const webrtc::PeerConnectionInterface::RTCConfiguration& Config)
+{
+	RTCConfig = Config;
+	OnConfig.Broadcast();
+}
+
+void UPixelStreamingSignallingComponent::OnSignallingSessionDescription(webrtc::SdpType Type, const FString& Sdp)
+{
+	if (Type == webrtc::SdpType::kOffer)
+	{
+		OnOffer.Broadcast(Sdp);
+	}
+	else if (Type == webrtc::SdpType::kAnswer)
+	{
+		// TODO if needed. Currently we never send an offer so we shouldnt expect an answer.
+	}
+}
+
+void UPixelStreamingSignallingComponent::OnSignallingRemoteIceCandidate(const FString& SdpMid, int SdpMLineIndex, const FString& Sdp)
+{
+	OnIceCandidate.Broadcast(SdpMid, SdpMLineIndex, Sdp);
+}
+
+void UPixelStreamingSignallingComponent::OnSignallingPeerDataChannels(int32 SendStreamId, int32 RecvStreamId)
+{
+	// TODO what to do with data channels.
+}
+
+void UPixelStreamingSignallingComponent::OnSignallingPlayerCount(uint32 Count)
+{
+}
diff --git a/Engine/Plugins/Media/PixelStreamingPlayer/Source/PixelStreamingPlayer/Public/PixelStreamingMediaTexture.h b/Engine/Plugins/Media/PixelStreamingPlayer/Source/PixelStreamingPlayer/Public/PixelStreamingMediaTexture.h
new file mode 100644
--- /dev/null
+++ b/Engine/Plugins/Media/PixelStreamingPlayer/Source/PixelStreamingPlayer/Public/PixelStreamingMediaTexture.h
@@ -0,0 +1,47 @@
+// Copyright Epic Games, Inc. All Rights Reserved.
+
+#pragma once
+
+#include "Engine/Texture.h"
+#include "RenderTargetPool.h"
+#include "api/media_stream_interface.h"
+#include "PixelStreamingMediaTexture.generated.h"
+
+class FPixelStreamingMediaTextureResource;
+
+/**
+ * A Texture Object that can be used in materials etc. that takes updates from webrtc frames.
+ */
+UCLASS(NotBlueprintType, NotBlueprintable, HideDropdown, HideCategories = (ImportSettings, Compression, Texture, Adjustments, Compositing, LevelOfDetail, Object), META = (DisplayName = "PixelStreaming Media Texture"))
+class PIXELSTREAMINGPLAYER_API UPixelStreamingMediaTexture : public UTexture, public rtc::VideoSinkInterface<webrtc::VideoFrame>
+{
+	GENERATED_UCLASS_BODY()
+
+protected:
+	// UTexture implementation
+	virtual void BeginDestroy() override;
+	virtual float GetSurfaceHeight() const override;
+	virtual float GetSurfaceWidth() const override;
+	virtual void GetResourceSizeEx(FResourceSizeEx& CumulativeResourceSize) override;
+	virtual EMaterialValueType GetMaterialType() const override;
+	virtual FTextureResource* CreateResource() override;
+
+	// from rtc::VideoSinkInterface<webrtc::VideoFrame>
+	virtual void OnFrame(const webrtc::VideoFrame& frame) override;
+
+	// updates the internal texture resource after each frame.
+	void UpdateTextureReference(FRHICommandList& RHICmdList, FTexture2DRHIRef Reference);
+
+private:
+	void InitializeResources();
+
+	FPixelStreamingMediaTextureResource* CurrentResource;
+
+	uint8_t* Buffer = nullptr;
+	size_t BufferSize = 0;
+
+	FCriticalSection RenderSyncContext;
+	FTextureRHIRef SourceTexture;
+	FPooledRenderTargetDesc RenderTargetDescriptor;
+	TRefCountPtr<IPooledRenderTarget> RenderTarget;
+};
diff --git a/Engine/Plugins/Media/PixelStreamingPlayer/Source/PixelStreamingPlayer/Public/PixelStreamingPeerComponent.h b/Engine/Plugins/Media/PixelStreamingPlayer/Source/PixelStreamingPlayer/Public/PixelStreamingPeerComponent.h
new file mode 100644
--- /dev/null
+++ b/Engine/Plugins/Media/PixelStreamingPlayer/Source/PixelStreamingPlayer/Public/PixelStreamingPeerComponent.h
@@ -0,0 +1,52 @@
+// Copyright Epic Games, Inc. All Rights Reserved.
+
+#pragma once
+
+#include "PixelStreamingPeerConnection.h"
+#include "PixelStreamingSignallingComponent.h"
+#include "PixelStreamingMediaTexture.h"
+#include "PixelStreamingPeerComponent.generated.h"
+
+/**
+ * A blueprint representation of a Pixel Streaming Peer Connection. Should communicate with a Pixel Streaming Signalling Connection
+ * and will accept video sinks to receive video data.
+ */
+UCLASS(BlueprintType, Blueprintable, Category = "PixelStreaming", META = (DisplayName = "PixelStreaming Peer Component", BlueprintSpawnableComponent))
+class PIXELSTREAMINGPLAYER_API UPixelStreamingPeerComponent : public UActorComponent
+{
+	GENERATED_UCLASS_BODY()
+
+public:
+	/**
+	 * Begins the peer connection.
+	 * @param InSignallingComponent A signalling component provided for the peer connection to communicate with. This should be be same signalling connection used to negotiate this peer connection.
+	 */
+	UFUNCTION(BlueprintCallable, Category = "PixelStreaming")
+	void Initialize(UPixelStreamingSignallingComponent* InSignallingComponent);
+
+	/**
+	 * Send an offer to the peer connection. This is used to negotiate a media connection.
+	 * @param Sdp The Session Description provided from signalling.
+	 */
+	UFUNCTION(BlueprintCallable, Category = "PixelStreaming")
+	void ReceiveOffer(const FString& Sdp);
+
+	/**
+	 * Notify the peer connection of an ICE candidate sent by the singalling connection.
+	 * @param SdpMid Provided by the singalling connection.
+	 * @param SdpMLineIndex Provided by the signalling connection.
+	 * @param Sdp Provided by the signalling connection.
+	 */
+	UFUNCTION(BlueprintCallable, Category = "PixelStreaming")
+	void ReceiveIceCandidate(const FString& SdpMid, int SdpMLineIndex, const FString& Sdp);
+
+	/**
+	 * A sink for the video data received once this connection has finished negotiating.
+	 */
+	UPROPERTY(EditDefaultsOnly, Category = "Properties", META = (DisplayName = "Pixel Streaming Video Sink", AllowPrivateAccess = true))
+	UPixelStreamingMediaTexture* VideoSink = nullptr;
+
+private:
+	TUniquePtr<FPixelStreamingPeerConnection> PeerConnection;
+	UPixelStreamingSignallingComponent* SignallingComponent = nullptr;
+};
diff --git a/Engine/Plugins/Media/PixelStreamingPlayer/Source/PixelStreamingPlayer/Public/PixelStreamingSignallingComponent.h b/Engine/Plugins/Media/PixelStreamingPlayer/Source/PixelStreamingPlayer/Public/PixelStreamingSignallingComponent.h
new file mode 100644
--- /dev/null
+++ b/Engine/Plugins/Media/PixelStreamingPlayer/Source/PixelStreamingPlayer/Public/PixelStreamingSignallingComponent.h
@@ -0,0 +1,112 @@
+// Copyright Epic Games, Inc. All Rights Reserved.
+
+#pragma once
+
+#include "PixelStreamingSignallingConnection.h"
+#include "StreamMediaSource.h"
+#include "Components/ActorComponent.h"
+#include "PixelStreamingSignallingComponent.generated.h"
+
+DECLARE_DYNAMIC_MULTICAST_SPARSE_DELEGATE(FPixelStreamingSignallingComponentConnected, UPixelStreamingSignallingComponent, OnConnected);
+DECLARE_DYNAMIC_MULTICAST_SPARSE_DELEGATE_OneParam(FPixelStreamingSignallingComponentConnectionError, UPixelStreamingSignallingComponent, OnConnectionError, const FString&, ErrorMsg);
+DECLARE_DYNAMIC_MULTICAST_SPARSE_DELEGATE_ThreeParams(FPixelStreamingSignallingComponentDisconnected, UPixelStreamingSignallingComponent, OnDisconnected, int32, StatusCode, const FString&, Reason, bool, bWasClean);
+DECLARE_DYNAMIC_MULTICAST_SPARSE_DELEGATE(FPixelStreamingSignallingComponentConfig, UPixelStreamingSignallingComponent, OnConfig);
+DECLARE_DYNAMIC_MULTICAST_SPARSE_DELEGATE_OneParam(FPixelStreamingSignallingComponentOffer, UPixelStreamingSignallingComponent, OnOffer, const FString&, Sdp);
+DECLARE_DYNAMIC_MULTICAST_SPARSE_DELEGATE_OneParam(FPixelStreamingSignallingComponentAnswer, UPixelStreamingSignallingComponent, OnAnswer, const FString&, Sdp);
+DECLARE_DYNAMIC_MULTICAST_SPARSE_DELEGATE_ThreeParams(FPixelStreamingSignallingComponentIceCandidate, UPixelStreamingSignallingComponent, OnIceCandidate, const FString&, SdpMid, int, SdpMLineIndex, const FString&, Sdp);
+DECLARE_DYNAMIC_MULTICAST_SPARSE_DELEGATE_TwoParams(FPixelStreamingSignallingComponentDataChannels, UPixelStreamingSignallingComponent, OnDataChannels, int32, SendStreamId, int32, RecvStreamId);
+
+/**
+ * A blueprint class representing a Pixel Streaming Signalling connection. Used to communicate with the signalling server and
+ * should route information to the peer connection.
+ */
+UCLASS(BlueprintType, Blueprintable, Category = "PixelStreaming", META = (DisplayName = "PixelStreaming Signalling Component", BlueprintSpawnableComponent))
+class PIXELSTREAMINGPLAYER_API UPixelStreamingSignallingComponent : public UActorComponent, public IPixelStreamingSignallingConnectionObserver
+{
+	GENERATED_UCLASS_BODY()
+
+public:
+	/**
+	 * Attempt to connect to a specified signalling server.
+	 * @param Url The url of the signalling server. Ignored if this component has a MediaSource. In that case the URL on the media source will be used instead.
+	 */
+	UFUNCTION(BlueprintCallable, Category = "PixelStreaming")
+	void Connect(const FString& Url);
+
+	/**
+	 * Disconnect from the signalling server. No action if no connection exists.
+	 */
+	UFUNCTION(BlueprintCallable, Category = "PixelStreaming")
+	void Disconnect();
+
+	/**
+	 * Fired when the signalling connection is successfully established.
+	 */
+	UPROPERTY(BlueprintAssignable, Category = "Components|Activation")
+	FPixelStreamingSignallingComponentConnected OnConnected;
+
+	/**
+	 * Fired if the connection failed or an error occurs during the connection. If this is fired at any point the connection should be considered closed.
+	 */
+	UPROPERTY(BlueprintAssignable, Category = "Components|Activation")
+	FPixelStreamingSignallingComponentConnectionError OnConnectionError;
+
+	/**
+	 * Fired when the connection successfully closes.
+	 */
+	UPROPERTY(BlueprintAssignable, Category = "Components|Activation")
+	FPixelStreamingSignallingComponentDisconnected OnDisconnected;
+
+	/**
+	 * Fired when the connection receives a config message from the server. This is the earliest place where the peer connection can be initialized.
+	 */
+	UPROPERTY(BlueprintAssignable, Category = "Components|Activation")
+	FPixelStreamingSignallingComponentConfig OnConfig;
+
+	/**
+	 * Fired when the connection receives an offer from the server. This means there is media being offered up to this connection. Forward to the peer connection.
+	 */
+	UPROPERTY(BlueprintAssignable, Category = "Components|Activation")
+	FPixelStreamingSignallingComponentOffer OnOffer;
+
+	/**
+	 * Fired when the server sends through an ice candidate. Forward this information on to the peer connection.
+	 */
+	UPROPERTY(BlueprintAssignable, Category = "Components|Activation")
+	FPixelStreamingSignallingComponentIceCandidate OnIceCandidate;
+
+	/**
+	 * If this media source is set we will use its supplied URL instead of the Url parameter on the connect call.
+	 */
+	UPROPERTY(EditDefaultsOnly, Category = "Properties", META = (DisplayName = "Stream Media Source", AllowPrivateAccess = true))
+	UStreamMediaSource* MediaSource = nullptr;
+
+	/**
+	 * Gets the current webrtc RTC configuration which is supplied by the signalling server.
+	 * @return The current webrtc RTC configuration from the signalling server.
+	 */
+	const webrtc::PeerConnectionInterface::RTCConfiguration& GetConfig() const { return RTCConfig; }
+
+	/**
+	 * Gets the current signalling connection object.
+	 * @return The raw connection object/
+	 */
+	FPixelStreamingSignallingConnection* GetConnection() const { return SignallingConnection.Get(); }
+
+protected:
+	//
+	// ISignallingServerConnectionObserver implementation.
+	//
+	virtual void OnSignallingConnected() override;
+	virtual void OnSignallingDisconnected(int32 StatusCode, const FString& Reason, bool bWasClean) override;
+	virtual void OnSignallingError(const FString& ErrorMsg) override;
+	virtual void OnSignallingConfig(const webrtc::PeerConnectionInterface::RTCConfiguration& Config) override;
+	virtual void OnSignallingSessionDescription(webrtc::SdpType Type, const FString& Sdp) override;
+	virtual void OnSignallingRemoteIceCandidate(const FString& SdpMid, int SdpMLineIndex, const FString& Sdp) override;
+	virtual void OnSignallingPeerDataChannels(int32 SendStreamId, int32 RecvStreamId) override;
+	virtual void OnSignallingPlayerCount(uint32 Count) override;
+
+private:
+	TUniquePtr<FPixelStreamingSignallingConnection> SignallingConnection;
+	webrtc::PeerConnectionInterface::RTCConfiguration RTCConfig;
+};
diff --git a/Engine/Plugins/Media/PixelStreamingPlayer/Source/PixelStreamingPlayerEditor/PixelStreamingPlayerEditor.Build.cs b/Engine/Plugins/Media/PixelStreamingPlayer/Source/PixelStreamingPlayerEditor/PixelStreamingPlayerEditor.Build.cs
new file mode 100644
--- /dev/null
+++ b/Engine/Plugins/Media/PixelStreamingPlayer/Source/PixelStreamingPlayerEditor/PixelStreamingPlayerEditor.Build.cs
@@ -0,0 +1,37 @@
+// Copyright Epic Games, Inc. All Rights Reserved.
+
+using System;
+using System.IO;
+using System.Collections.Generic;
+using EpicGames.Core;
+
+namespace UnrealBuildTool.Rules
+{
+	public class PixelStreamingPlayerEditor : ModuleRules
+	{
+		public PixelStreamingPlayerEditor(ReadOnlyTargetRules Target) : base(Target)
+		{
+			var EngineDir = Path.GetFullPath(Target.RelativeEnginePath);
+
+			PublicIncludePaths.AddRange(
+				new string[] {
+				});
+
+			PrivateIncludePaths.AddRange(
+				new string[] {
+				});
+
+			PublicDependencyModuleNames.AddRange(
+				new string[] {
+				});
+
+			PrivateDependencyModuleNames.AddRange(
+				new string[] {
+					"Core",
+					"CoreUObject",
+					"PixelStreamingPlayer",
+					"UnrealEd"
+				});
+		}
+	}
+}
diff --git a/Engine/Plugins/Media/PixelStreamingPlayer/Source/PixelStreamingPlayerEditor/Private/PixelStreamingMediaTextureFactory.cpp b/Engine/Plugins/Media/PixelStreamingPlayer/Source/PixelStreamingPlayerEditor/Private/PixelStreamingMediaTextureFactory.cpp
new file mode 100644
--- /dev/null
+++ b/Engine/Plugins/Media/PixelStreamingPlayer/Source/PixelStreamingPlayerEditor/Private/PixelStreamingMediaTextureFactory.cpp
@@ -0,0 +1,40 @@
+// Copyright Epic Games, Inc. All Rights Reserved.
+
+#include "PixelStreamingMediaTextureFactory.h"
+
+#include <AssetTypeCategories.h>
+#include "PixelStreamingMediaTexture.h"
+
+#define LOCTEXT_NAMESPACE "PixelStreamingPlayerEditorMediaTextureFactory"
+
+UPixelStreamingMediaTextureFactory::UPixelStreamingMediaTextureFactory(const FObjectInitializer& ObjectInitializer)
+	: Super(ObjectInitializer)
+{
+	bCreateNew = true;
+	bEditAfterNew = true;
+
+	SupportedClass = UPixelStreamingMediaTexture::StaticClass();
+}
+
+FText UPixelStreamingMediaTextureFactory::GetDisplayName() const
+{
+	return LOCTEXT("PixelStreamingMediaTextureFactoryDisplayName", "Pixel Streaming Media Texture");
+}
+
+uint32 UPixelStreamingMediaTextureFactory::GetMenuCategories() const
+{
+	return EAssetTypeCategories::Textures;
+}
+
+UObject* UPixelStreamingMediaTextureFactory::FactoryCreateNew(UClass* InClass, UObject* InParent, FName InName, EObjectFlags Flags, UObject* Context, FFeedbackContext* Warn)
+{
+	if (UPixelStreamingMediaTexture* Resource = NewObject<UPixelStreamingMediaTexture>(InParent, InName, Flags | RF_Transactional))
+	{
+		Resource->UpdateResource();
+		return Resource;
+	}
+
+	return nullptr;
+}
+
+#undef LOCTEXT_NAMESPACE
diff --git a/Engine/Plugins/Media/PixelStreamingPlayer/Source/PixelStreamingPlayerEditor/Private/PixelStreamingMediaTextureFactory.h b/Engine/Plugins/Media/PixelStreamingPlayer/Source/PixelStreamingPlayerEditor/Private/PixelStreamingMediaTextureFactory.h
new file mode 100644
--- /dev/null
+++ b/Engine/Plugins/Media/PixelStreamingPlayer/Source/PixelStreamingPlayerEditor/Private/PixelStreamingMediaTextureFactory.h
@@ -0,0 +1,20 @@
+// Copyright Epic Games, Inc. All Rights Reserved.
+
+#pragma once
+
+#include <CoreMinimal.h>
+#include <Factories/Factory.h>
+#include <UObject/Object.h>
+
+#include "PixelStreamingMediaTextureFactory.generated.h"
+
+UCLASS()
+class UPixelStreamingMediaTextureFactory : public UFactory
+{
+	GENERATED_UCLASS_BODY()
+
+	public:
+		virtual FText GetDisplayName() const override;
+		virtual uint32 GetMenuCategories() const override;
+		virtual UObject* FactoryCreateNew(UClass* InClass, UObject* InParent, FName InName, EObjectFlags Flags, UObject* Context, FFeedbackContext* Warn) override;
+};
diff --git a/Engine/Plugins/Media/PixelStreamingPlayer/Source/PixelStreamingPlayerEditor/Private/PixelStreamingPlayerEditorModule.cpp b/Engine/Plugins/Media/PixelStreamingPlayer/Source/PixelStreamingPlayerEditor/Private/PixelStreamingPlayerEditorModule.cpp
new file mode 100644
--- /dev/null
+++ b/Engine/Plugins/Media/PixelStreamingPlayer/Source/PixelStreamingPlayerEditor/Private/PixelStreamingPlayerEditorModule.cpp
@@ -0,0 +1,14 @@
+// Copyright Epic Games, Inc. All Rights Reserved.
+
+#include "Modules/ModuleManager.h"
+
+class FPixelStreamingPlayerEditorModule : public IModuleInterface
+{
+public:
+private:
+	/** IModuleInterface implementation */
+	void StartupModule() override {}
+	void ShutdownModule() override {}
+};
+
+IMPLEMENT_MODULE(FPixelStreamingPlayerEditorModule, PixelStreamingPlayerEditor)
