Added missing implementations for pure abstract functions.

diff --git a/Engine/Plugins/Media/PixelStreamingPlayer/Source/PixelStreamingPlayer/Private/PixelStreamingMediaTexture.cpp b/Engine/Plugins/Media/PixelStreamingPlayer/Source/PixelStreamingPlayer/Private/PixelStreamingMediaTexture.cpp
--- a/Engine/Plugins/Media/PixelStreamingPlayer/Source/PixelStreamingPlayer/Private/PixelStreamingMediaTexture.cpp
+++ b/Engine/Plugins/Media/PixelStreamingPlayer/Source/PixelStreamingPlayer/Private/PixelStreamingMediaTexture.cpp
@@ -84,14 +84,14 @@ void UPixelStreamingMediaTexture::BeginDestroy()
 	Super::BeginDestroy();
 }
 
-float UPixelStreamingMediaTexture::GetSurfaceHeight() const
+float UPixelStreamingMediaTexture::GetSurfaceWidth() const
 {
-	return CurrentResource != nullptr ? CurrentResource->GetSizeY() : 0.0f;
+	return CurrentResource != nullptr ? CurrentResource->GetSizeX() : 0.0f;
 }
 
-float UPixelStreamingMediaTexture::GetSurfaceWidth() const
+float UPixelStreamingMediaTexture::GetSurfaceHeight() const
 {
-	return CurrentResource != nullptr ? CurrentResource->GetSizeX() : 0.0f;
+	return CurrentResource != nullptr ? CurrentResource->GetSizeY() : 0.0f;
 }
 
 void UPixelStreamingMediaTexture::GetResourceSizeEx(FResourceSizeEx& CumulativeResourceSize)
diff --git a/Engine/Plugins/Media/PixelStreamingPlayer/Source/PixelStreamingPlayer/Public/PixelStreamingMediaTexture.h b/Engine/Plugins/Media/PixelStreamingPlayer/Source/PixelStreamingPlayer/Public/PixelStreamingMediaTexture.h
--- a/Engine/Plugins/Media/PixelStreamingPlayer/Source/PixelStreamingPlayer/Public/PixelStreamingMediaTexture.h
+++ b/Engine/Plugins/Media/PixelStreamingPlayer/Source/PixelStreamingPlayer/Public/PixelStreamingMediaTexture.h
@@ -19,12 +19,19 @@ class PIXELSTREAMINGPLAYER_API UPixelStreamingMediaTexture : public UTexture, pu
 
 protected:
 	// UTexture implementation
-	virtual void BeginDestroy() override;
-	virtual float GetSurfaceHeight() const override;
+
+	virtual FTextureResource* CreateResource() override;
+	virtual EMaterialValueType GetMaterialType() const override;
 	virtual float GetSurfaceWidth() const override;
+	virtual float GetSurfaceHeight() const override;
+	virtual float GetSurfaceDepth() const override { return 0; }
+	virtual uint32 GetSurfaceArraySize() const override { return 0; }
+	virtual ETextureClass GetTextureClass() const { return ETextureClass::Other2DNoSource; }
+
 	virtual void GetResourceSizeEx(FResourceSizeEx& CumulativeResourceSize) override;
-	virtual EMaterialValueType GetMaterialType() const override;
-	virtual FTextureResource* CreateResource() override;
+
+	// UObject overrides.
+	virtual void BeginDestroy() override;
 
 	// from rtc::VideoSinkInterface<webrtc::VideoFrame>
 	virtual void OnFrame(const webrtc::VideoFrame& frame) override;
