Removed PixelStreamingPlayer private pch.
Timing when compiling module with PCH(including compiling PCH): 17.97 secs
Timing using shared PCH: 6.91 secs

diff --git a/Engine/Plugins/Media/PixelStreamingPlayer/Source/PixelStreamingPlayer/PixelStreamingPlayer.Build.cs b/Engine/Plugins/Media/PixelStreamingPlayer/Source/PixelStreamingPlayer/PixelStreamingPlayer.Build.cs
--- a/Engine/Plugins/Media/PixelStreamingPlayer/Source/PixelStreamingPlayer/PixelStreamingPlayer.Build.cs
+++ b/Engine/Plugins/Media/PixelStreamingPlayer/Source/PixelStreamingPlayer/PixelStreamingPlayer.Build.cs
@@ -11,9 +11,6 @@ namespace UnrealBuildTool.Rules
 	{
 		public PixelStreamingPlayer(ReadOnlyTargetRules Target) : base(Target)
 		{
-			PCHUsage = PCHUsageMode.UseExplicitOrSharedPCHs;
-			PrivatePCHHeaderFile = "Private/PCH.h";
-
 			var EngineDir = Path.GetFullPath(Target.RelativeEnginePath);
 
 			PublicIncludePaths.AddRange(
diff --git a/Engine/Plugins/Media/PixelStreamingPlayer/Source/PixelStreamingPlayer/Private/PCH.h b/Engine/Plugins/Media/PixelStreamingPlayer/Source/PixelStreamingPlayer/Private/PCH.h
deleted file mode 100644
