[Backout] - CL21684160
Original CL Desc
-----------------------------------------------------------------
Removed PixelStreamingPlayer private pch.
Timing when compiling module with PCH(including compiling PCH): 17.97 secs
Timing using shared PCH: 6.91 secs

diff --git a/Engine/Plugins/Media/PixelStreamingPlayer/Source/PixelStreamingPlayer/PixelStreamingPlayer.Build.cs b/Engine/Plugins/Media/PixelStreamingPlayer/Source/PixelStreamingPlayer/PixelStreamingPlayer.Build.cs
--- a/Engine/Plugins/Media/PixelStreamingPlayer/Source/PixelStreamingPlayer/PixelStreamingPlayer.Build.cs
+++ b/Engine/Plugins/Media/PixelStreamingPlayer/Source/PixelStreamingPlayer/PixelStreamingPlayer.Build.cs
@@ -11,6 +11,9 @@ namespace UnrealBuildTool.Rules
 	{
 		public PixelStreamingPlayer(ReadOnlyTargetRules Target) : base(Target)
 		{
+			PCHUsage = PCHUsageMode.UseExplicitOrSharedPCHs;
+			PrivatePCHHeaderFile = "Private/PCH.h";
+
 			var EngineDir = Path.GetFullPath(Target.RelativeEnginePath);
 
 			PublicIncludePaths.AddRange(
diff --git a/Engine/Plugins/Media/PixelStreamingPlayer/Source/PixelStreamingPlayer/Private/PCH.h b/Engine/Plugins/Media/PixelStreamingPlayer/Source/PixelStreamingPlayer/Private/PCH.h
new file mode 100644
--- /dev/null
+++ b/Engine/Plugins/Media/PixelStreamingPlayer/Source/PixelStreamingPlayer/Private/PCH.h
@@ -0,0 +1,7 @@
+// Copyright Epic Games, Inc. All Rights Reserved.
+
+#pragma once
+
+#include "WebSocketsModule.h"
+#include "WebRTCIncludes.h"
+#include "PixelStreamingPlayerPrivate.h"
