Fixed non unity vs2019 farm compile errors

diff --git a/Engine/Plugins/Media/PixelStreamingPlayer/Source/PixelStreamingPlayer/Private/PixelStreamingPlayerModule.cpp b/Engine/Plugins/Media/PixelStreamingPlayer/Source/PixelStreamingPlayer/Private/PixelStreamingPlayerModule.cpp
--- a/Engine/Plugins/Media/PixelStreamingPlayer/Source/PixelStreamingPlayer/Private/PixelStreamingPlayerModule.cpp
+++ b/Engine/Plugins/Media/PixelStreamingPlayer/Source/PixelStreamingPlayer/Private/PixelStreamingPlayerModule.cpp
@@ -3,7 +3,7 @@
 #include "PixelStreamingPlayerPrivate.h"
 #include "Modules/ModuleManager.h"
 #include "WebSocketsModule.h"
-#include <rtc_base/ssl_adapter.h>
+#include "WebRTCIncludes.h" // IWYU pragma: keep
 
 DEFINE_LOG_CATEGORY(LogPixelStreamingPlayer);
 
diff --git a/Engine/Plugins/NetcodeUnitTest/NetcodeUnitTest/Source/NetcodeUnitTest/Private/UnitTests/Engine/PackedVectorTest.cpp b/Engine/Plugins/NetcodeUnitTest/NetcodeUnitTest/Source/NetcodeUnitTest/Private/UnitTests/Engine/PackedVectorTest.cpp
--- a/Engine/Plugins/NetcodeUnitTest/NetcodeUnitTest/Source/NetcodeUnitTest/Private/UnitTests/Engine/PackedVectorTest.cpp
+++ b/Engine/Plugins/NetcodeUnitTest/NetcodeUnitTest/Source/NetcodeUnitTest/Private/UnitTests/Engine/PackedVectorTest.cpp
@@ -5,6 +5,7 @@
 #include "NetcodeUnitTest.h"
 #include "Serialization/BitReader.h"
 #include "Serialization/BitWriter.h"
+#include <limits>
 
 namespace PackedVectorTest
 {
