Fixed error where UPixelStreamingSignallingComponent instances were getting deleted twice due to shared pointer of this.

diff --git a/Engine/Plugins/Experimental/PixelStreamingPlayer/Source/PixelStreamingPlayer/Private/PixelStreamingSignallingComponent.cpp b/Engine/Plugins/Experimental/PixelStreamingPlayer/Source/PixelStreamingPlayer/Private/PixelStreamingSignallingComponent.cpp
--- a/Engine/Plugins/Experimental/PixelStreamingPlayer/Source/PixelStreamingPlayer/Private/PixelStreamingSignallingComponent.cpp
+++ b/Engine/Plugins/Experimental/PixelStreamingPlayer/Source/PixelStreamingPlayer/Private/PixelStreamingSignallingComponent.cpp
@@ -11,7 +11,8 @@
 UPixelStreamingSignallingComponent::UPixelStreamingSignallingComponent(const FObjectInitializer& ObjectInitializer)
 	: Super(ObjectInitializer)
 {
-	SignallingConnection = MakeUnique<FPixelStreamingSignallingConnection>(MakeShareable(this));
+	SignallingServerConnectionObserver = MakeShared<FPixelStreamingSignallingConnectionObserver>(this);
+	SignallingConnection = MakeUnique<FPixelStreamingSignallingConnection>(SignallingServerConnectionObserver);
 	// I dont disable keep alive here because it causes a deadlock.
 	// This is because the game thread is waiting to load this component but the timer must be set
 	// on the game thread.
@@ -65,69 +66,69 @@ void UPixelStreamingSignallingComponent::SendIceCandidate(const FPixelStreamingI
 	SignallingConnection->SendIceCandidate(*CandidateWrapper.ToWebRTC());
 }
 
-void UPixelStreamingSignallingComponent::OnSignallingConnected()
+void UPixelStreamingSignallingComponent::FPixelStreamingSignallingConnectionObserver::OnSignallingConnected()
 {
-	OnConnected.Broadcast();
+	Parent->OnConnected.Broadcast();
 }
 
-void UPixelStreamingSignallingComponent::OnSignallingDisconnected(int32 StatusCode, const FString& Reason, bool bWasClean)
+void UPixelStreamingSignallingComponent::FPixelStreamingSignallingConnectionObserver::OnSignallingDisconnected(int32 StatusCode, const FString& Reason, bool bWasClean)
 {
-	OnDisconnected.Broadcast(StatusCode, Reason, bWasClean);
+	Parent->OnDisconnected.Broadcast(StatusCode, Reason, bWasClean);
 }
 
-void UPixelStreamingSignallingComponent::OnSignallingError(const FString& ErrorMsg)
+void UPixelStreamingSignallingComponent::FPixelStreamingSignallingConnectionObserver::OnSignallingError(const FString& ErrorMsg)
 {
-	OnConnectionError.Broadcast(ErrorMsg);
+	Parent->OnConnectionError.Broadcast(ErrorMsg);
 }
 
-void UPixelStreamingSignallingComponent::OnSignallingConfig(const webrtc::PeerConnectionInterface::RTCConfiguration& Config)
+void UPixelStreamingSignallingComponent::FPixelStreamingSignallingConnectionObserver::OnSignallingConfig(const webrtc::PeerConnectionInterface::RTCConfiguration& Config)
 {
 	FPixelStreamingRTCConfigWrapper Wrapper;
 	Wrapper.Config = Config;
-	OnConfig.Broadcast(Wrapper);
+	Parent->OnConfig.Broadcast(Wrapper);
 }
 
-void UPixelStreamingSignallingComponent::OnSignallingSessionDescription(webrtc::SdpType Type, const FString& Sdp)
+void UPixelStreamingSignallingComponent::FPixelStreamingSignallingConnectionObserver::OnSignallingSessionDescription(webrtc::SdpType Type, const FString& Sdp)
 {
 	if (Type == webrtc::SdpType::kOffer)
 	{
-		OnOffer.Broadcast(Sdp);
+		Parent->OnOffer.Broadcast(Sdp);
 	}
 	else if (Type == webrtc::SdpType::kAnswer)
 	{
-		OnAnswer.Broadcast(Sdp);
+		Parent->OnAnswer.Broadcast(Sdp);
 	}
 }
 
-void UPixelStreamingSignallingComponent::OnSignallingRemoteIceCandidate(const FString& SdpMid, int SdpMLineIndex, const FString& Sdp)
+void UPixelStreamingSignallingComponent::FPixelStreamingSignallingConnectionObserver::OnSignallingRemoteIceCandidate(const FString& SdpMid, int SdpMLineIndex, const FString& Sdp)
 {
 	FPixelStreamingIceCandidateWrapper CandidateWrapper(SdpMid, SdpMLineIndex, Sdp);
-	OnIceCandidate.Broadcast(CandidateWrapper);
+	Parent->OnIceCandidate.Broadcast(CandidateWrapper);
 }
 
-void UPixelStreamingSignallingComponent::OnSignallingPeerDataChannels(int32 SendStreamId, int32 RecvStreamId)
+void UPixelStreamingSignallingComponent::FPixelStreamingSignallingConnectionObserver::OnSignallingPeerDataChannels(int32 SendStreamId, int32 RecvStreamId)
 {
 	// TODO (Matthew.Cotton): What to do with data channels?
 }
 
-void UPixelStreamingSignallingComponent::OnSignallingPlayerCount(uint32 Count)
+void UPixelStreamingSignallingComponent::FPixelStreamingSignallingConnectionObserver::OnSignallingPlayerCount(uint32 Count)
 {
 	// no-op player count is not interesting to us here
 }
 
-void UPixelStreamingSignallingComponent::OnSignallingPlayerConnected(FPixelStreamingPlayerId PlayerId, const FPixelStreamingPlayerConfig& PlayerConfig)
+void UPixelStreamingSignallingComponent::FPixelStreamingSignallingConnectionObserver::OnSignallingPlayerConnected(FPixelStreamingPlayerId PlayerId, const FPixelStreamingPlayerConfig& PlayerConfig)
 {
 	// no-op our player receiving other players is possible but not a case we are interested in
 	UE_LOG(LogPixelStreamingPlayer, Warning, TEXT("Warning: Received message about player connected from SS, this very likely means you have misconfigured the Pixel Streaming player connection and entered the streamer port. Please check your connect node and ensure it using the player port, e.g. 80 (or no port)."));
 }
 
-void UPixelStreamingSignallingComponent::OnSignallingPlayerDisconnected(FPixelStreamingPlayerId PlayerId)
+void UPixelStreamingSignallingComponent::FPixelStreamingSignallingConnectionObserver::OnSignallingPlayerDisconnected(FPixelStreamingPlayerId PlayerId)
 {
 	// no-op our player knowing when other players disconnect is not interesting for now
 	UE_LOG(LogPixelStreamingPlayer, Warning, TEXT("Warning: Received message about player disconnected from SS, this very likely means you have misconfigured the Pixel Streaming player connection and entered the streamer port. Please check your connect node and ensure it using the player port, e.g. 80 (or no port)."));
 }
 
-void UPixelStreamingSignallingComponent::OnSignallingSFUPeerDataChannels(FPixelStreamingPlayerId SFUId, FPixelStreamingPlayerId PlayerId, int32 SendStreamId, int32 RecvStreamId)
+void UPixelStreamingSignallingComponent::FPixelStreamingSignallingConnectionObserver::OnSignallingSFUPeerDataChannels(FPixelStreamingPlayerId SFUId, FPixelStreamingPlayerId PlayerId, int32 SendStreamId, int32 RecvStreamId)
 {
 	// TODO (Matthew.Cotton): What to do with data channels?
 }
diff --git a/Engine/Plugins/Experimental/PixelStreamingPlayer/Source/PixelStreamingPlayer/Public/PixelStreamingSignallingComponent.h b/Engine/Plugins/Experimental/PixelStreamingPlayer/Source/PixelStreamingPlayer/Public/PixelStreamingSignallingComponent.h
--- a/Engine/Plugins/Experimental/PixelStreamingPlayer/Source/PixelStreamingPlayer/Public/PixelStreamingSignallingComponent.h
+++ b/Engine/Plugins/Experimental/PixelStreamingPlayer/Source/PixelStreamingPlayer/Public/PixelStreamingSignallingComponent.h
@@ -22,7 +22,7 @@ DECLARE_DYNAMIC_MULTICAST_SPARSE_DELEGATE_TwoParams(FPixelStreamingSignallingCom
  * should route information to the peer connection.
  */
 UCLASS(BlueprintType, Blueprintable, Category = "PixelStreaming", META = (DisplayName = "PixelStreaming Signalling Component", BlueprintSpawnableComponent))
-class PIXELSTREAMINGPLAYER_API UPixelStreamingSignallingComponent : public UActorComponent, public IPixelStreamingSignallingConnectionObserver
+class PIXELSTREAMINGPLAYER_API UPixelStreamingSignallingComponent : public UActorComponent
 {
 	GENERATED_UCLASS_BODY()
 
@@ -107,26 +107,36 @@ public:
 	 * If this media source is set we will use its supplied URL instead of the Url parameter on the connect call.
 	 */
 	UPROPERTY(EditDefaultsOnly, Category = "Properties", META = (DisplayName = "Stream Media Source", AllowPrivateAccess = true))
-	TObjectPtr<UStreamMediaSource> MediaSource = nullptr;
+	TObjectPtr<UStreamMediaSource> MediaSource;
 
 protected:
-	//
-	// ISignallingServerConnectionObserver implementation.
-	//
-	virtual void OnSignallingConnected() override;
-	virtual void OnSignallingDisconnected(int32 StatusCode, const FString& Reason, bool bWasClean) override;
-	virtual void OnSignallingError(const FString& ErrorMsg) override;
-	virtual void OnSignallingConfig(const webrtc::PeerConnectionInterface::RTCConfiguration& Config) override;
-	virtual void OnSignallingSessionDescription(webrtc::SdpType Type, const FString& Sdp) override;
-	virtual void OnSignallingRemoteIceCandidate(const FString& SdpMid, int SdpMLineIndex, const FString& Sdp) override;
-	virtual void OnSignallingPeerDataChannels(int32 SendStreamId, int32 RecvStreamId) override;
-	virtual void OnSignallingPlayerCount(uint32 Count) override;
-	virtual void OnSignallingPlayerConnected(FPixelStreamingPlayerId PlayerId, const FPixelStreamingPlayerConfig& PlayerConfig) override;
-	virtual void OnSignallingPlayerDisconnected(FPixelStreamingPlayerId PlayerId) override;
-	virtual void OnSignallingSFUPeerDataChannels(FPixelStreamingPlayerId SFUId, FPixelStreamingPlayerId PlayerId, int32 SendStreamId, int32 RecvStreamId) override;
-	virtual void OnSignallingSessionDescription(FPixelStreamingPlayerId PlayerId, webrtc::SdpType Type, const FString& Sdp) override {}
-	virtual void OnSignallingRemoteIceCandidate(FPixelStreamingPlayerId PlayerId, const FString& SdpMid, int SdpMLineIndex, const FString& Sdp) override {}
+	class FPixelStreamingSignallingConnectionObserver : public IPixelStreamingSignallingConnectionObserver
+	{
+	public:
+		FPixelStreamingSignallingConnectionObserver(UPixelStreamingSignallingComponent* InParent) : Parent(InParent) {}
+
+		//
+		// ISignallingServerConnectionObserver implementation.
+		//
+		virtual void OnSignallingConnected() override;
+		virtual void OnSignallingDisconnected(int32 StatusCode, const FString& Reason, bool bWasClean) override;
+		virtual void OnSignallingError(const FString& ErrorMsg) override;
+		virtual void OnSignallingConfig(const webrtc::PeerConnectionInterface::RTCConfiguration& Config) override;
+		virtual void OnSignallingSessionDescription(webrtc::SdpType Type, const FString& Sdp) override;
+		virtual void OnSignallingRemoteIceCandidate(const FString& SdpMid, int SdpMLineIndex, const FString& Sdp) override;
+		virtual void OnSignallingPeerDataChannels(int32 SendStreamId, int32 RecvStreamId) override;
+		virtual void OnSignallingPlayerCount(uint32 Count) override;
+		virtual void OnSignallingPlayerConnected(FPixelStreamingPlayerId PlayerId, const FPixelStreamingPlayerConfig& PlayerConfig) override;
+		virtual void OnSignallingPlayerDisconnected(FPixelStreamingPlayerId PlayerId) override;
+		virtual void OnSignallingSFUPeerDataChannels(FPixelStreamingPlayerId SFUId, FPixelStreamingPlayerId PlayerId, int32 SendStreamId, int32 RecvStreamId) override;
+		virtual void OnSignallingSessionDescription(FPixelStreamingPlayerId PlayerId, webrtc::SdpType Type, const FString& Sdp) override {}
+		virtual void OnSignallingRemoteIceCandidate(FPixelStreamingPlayerId PlayerId, const FString& SdpMid, int SdpMLineIndex, const FString& Sdp) override {}
+	
+	private:
+		TObjectPtr<UPixelStreamingSignallingComponent> Parent;
+	};
 
 private:
 	TUniquePtr<IPixelStreamingSignallingConnection> SignallingConnection;
+	TSharedPtr<IPixelStreamingSignallingConnectionObserver> SignallingServerConnectionObserver;
 };
