Merging CL24224409 from Dev-TensorWorks:

Adding list/subscribe/unsubscribe functionality into the player side of pixel streaming.

diff --git a/Engine/Plugins/Experimental/PixelStreamingPlayer/Source/PixelStreamingPlayer/Private/AsyncActionGetStreamers.cpp b/Engine/Plugins/Experimental/PixelStreamingPlayer/Source/PixelStreamingPlayer/Private/AsyncActionGetStreamers.cpp
new file mode 100644
--- /dev/null
+++ b/Engine/Plugins/Experimental/PixelStreamingPlayer/Source/PixelStreamingPlayer/Private/AsyncActionGetStreamers.cpp
@@ -0,0 +1,19 @@
+// Copyright Epic Games, Inc. All Rights Reserved.
+
+#include "AsyncActionGetStreamers.h"
+
+void UAsyncAction_GetStreamers::Activate()
+{
+	SignallingConnectionComponent->AsyncRequestStreamerList([this](const TArray<FString>& StreamerList) {
+		Completed.Broadcast(StreamerList);
+		SetReadyToDestroy();
+	});
+}
+
+UAsyncAction_GetStreamers* UAsyncAction_GetStreamers::GetStreamerIdList(UPixelStreamingSignallingComponent* SignallingComponent)
+{
+	UAsyncAction_GetStreamers* Action = NewObject<UAsyncAction_GetStreamers>();
+	Action->SignallingConnectionComponent = SignallingComponent;
+	Action->RegisterWithGameInstance(SignallingComponent);
+	return Action;
+}
diff --git a/Engine/Plugins/Experimental/PixelStreamingPlayer/Source/PixelStreamingPlayer/Private/AsyncActionGetStreamers.h b/Engine/Plugins/Experimental/PixelStreamingPlayer/Source/PixelStreamingPlayer/Private/AsyncActionGetStreamers.h
new file mode 100644
--- /dev/null
+++ b/Engine/Plugins/Experimental/PixelStreamingPlayer/Source/PixelStreamingPlayer/Private/AsyncActionGetStreamers.h
@@ -0,0 +1,25 @@
+// Copyright Epic Games, Inc. All Rights Reserved.
+
+#pragma once
+
+#include "PixelStreamingSignallingComponent.h"
+#include "AsyncActionGetStreamers.generated.h"
+
+DECLARE_DYNAMIC_MULTICAST_DELEGATE_OneParam(FOnStreamerListCompleted, const TArray<FString>&, StreamerList);
+
+UCLASS()
+class PIXELSTREAMINGPLAYER_API UAsyncAction_GetStreamers : public UBlueprintAsyncActionBase
+{
+	GENERATED_BODY()
+
+public:
+	virtual void Activate() override;
+
+	UFUNCTION(BlueprintCallable, meta = (BlueprintInternalUseOnly = "true", Category = "PixelStreaming"))
+	static UAsyncAction_GetStreamers* GetStreamerIdList(UPixelStreamingSignallingComponent* SignallingComponent);
+
+	UPROPERTY(BlueprintAssignable)
+	FOnStreamerListCompleted Completed;
+
+	UPixelStreamingSignallingComponent* SignallingConnectionComponent;
+};
diff --git a/Engine/Plugins/Experimental/PixelStreamingPlayer/Source/PixelStreamingPlayer/Private/PixelStreamingSignallingComponent.cpp b/Engine/Plugins/Experimental/PixelStreamingPlayer/Source/PixelStreamingPlayer/Private/PixelStreamingSignallingComponent.cpp
--- a/Engine/Plugins/Experimental/PixelStreamingPlayer/Source/PixelStreamingPlayer/Private/PixelStreamingSignallingComponent.cpp
+++ b/Engine/Plugins/Experimental/PixelStreamingPlayer/Source/PixelStreamingPlayer/Private/PixelStreamingSignallingComponent.cpp
@@ -11,7 +11,7 @@
 UPixelStreamingSignallingComponent::UPixelStreamingSignallingComponent(const FObjectInitializer& ObjectInitializer)
 	: Super(ObjectInitializer)
 {
-	SignallingServerConnectionObserver = MakeShared<FPixelStreamingSignallingConnectionObserver>(this);
+	SignallingServerConnectionObserver = MakeShared<FSignallingObserver>(this);
 	SignallingConnection = MakeUnique<FPixelStreamingSignallingConnection>(SignallingServerConnectionObserver);
 	// I dont disable keep alive here because it causes a deadlock.
 	// This is because the game thread is waiting to load this component but the timer must be set
@@ -37,6 +37,16 @@ void UPixelStreamingSignallingComponent::Disconnect()
 	SignallingConnection->Disconnect();
 }
 
+void UPixelStreamingSignallingComponent::Subscribe(const FString& StreamerId)
+{
+	SignallingConnection->SendSubscribe(StreamerId);
+}
+
+void UPixelStreamingSignallingComponent::Unsubscribe()
+{
+	SignallingConnection->SendUnsubscribe();
+}
+
 void UPixelStreamingSignallingComponent::SendOffer(const FPixelStreamingSessionDescriptionWrapper& Offer)
 {
 	if (Offer.SDP == nullptr)
@@ -66,29 +76,29 @@ void UPixelStreamingSignallingComponent::SendIceCandidate(const FPixelStreamingI
 	SignallingConnection->SendIceCandidate(*CandidateWrapper.ToWebRTC());
 }
 
-void UPixelStreamingSignallingComponent::FPixelStreamingSignallingConnectionObserver::OnSignallingConnected()
+void UPixelStreamingSignallingComponent::FSignallingObserver::OnSignallingConnected()
 {
 	Parent->OnConnected.Broadcast();
 }
 
-void UPixelStreamingSignallingComponent::FPixelStreamingSignallingConnectionObserver::OnSignallingDisconnected(int32 StatusCode, const FString& Reason, bool bWasClean)
+void UPixelStreamingSignallingComponent::FSignallingObserver::OnSignallingDisconnected(int32 StatusCode, const FString& Reason, bool bWasClean)
 {
 	Parent->OnDisconnected.Broadcast(StatusCode, Reason, bWasClean);
 }
 
-void UPixelStreamingSignallingComponent::FPixelStreamingSignallingConnectionObserver::OnSignallingError(const FString& ErrorMsg)
+void UPixelStreamingSignallingComponent::FSignallingObserver::OnSignallingError(const FString& ErrorMsg)
 {
 	Parent->OnConnectionError.Broadcast(ErrorMsg);
 }
 
-void UPixelStreamingSignallingComponent::FPixelStreamingSignallingConnectionObserver::OnSignallingConfig(const webrtc::PeerConnectionInterface::RTCConfiguration& Config)
+void UPixelStreamingSignallingComponent::FSignallingObserver::OnSignallingConfig(const webrtc::PeerConnectionInterface::RTCConfiguration& Config)
 {
 	FPixelStreamingRTCConfigWrapper Wrapper;
 	Wrapper.Config = Config;
 	Parent->OnConfig.Broadcast(Wrapper);
 }
 
-void UPixelStreamingSignallingComponent::FPixelStreamingSignallingConnectionObserver::OnSignallingSessionDescription(webrtc::SdpType Type, const FString& Sdp)
+void UPixelStreamingSignallingComponent::FSignallingObserver::OnSignallingSessionDescription(webrtc::SdpType Type, const FString& Sdp)
 {
 	if (Type == webrtc::SdpType::kOffer)
 	{
@@ -100,35 +110,46 @@ void UPixelStreamingSignallingComponent::FPixelStreamingSignallingConnectionObse
 	}
 }
 
-void UPixelStreamingSignallingComponent::FPixelStreamingSignallingConnectionObserver::OnSignallingRemoteIceCandidate(const FString& SdpMid, int SdpMLineIndex, const FString& Sdp)
+void UPixelStreamingSignallingComponent::FSignallingObserver::OnSignallingRemoteIceCandidate(const FString& SdpMid, int SdpMLineIndex, const FString& Sdp)
 {
 	FPixelStreamingIceCandidateWrapper CandidateWrapper(SdpMid, SdpMLineIndex, Sdp);
 	Parent->OnIceCandidate.Broadcast(CandidateWrapper);
 }
 
-void UPixelStreamingSignallingComponent::FPixelStreamingSignallingConnectionObserver::OnSignallingPeerDataChannels(int32 SendStreamId, int32 RecvStreamId)
+void UPixelStreamingSignallingComponent::FSignallingObserver::OnSignallingPeerDataChannels(int32 SendStreamId, int32 RecvStreamId)
 {
 	// TODO (Matthew.Cotton): What to do with data channels?
 }
 
-void UPixelStreamingSignallingComponent::FPixelStreamingSignallingConnectionObserver::OnSignallingPlayerCount(uint32 Count)
+void UPixelStreamingSignallingComponent::FSignallingObserver::OnSignallingPlayerCount(uint32 Count)
 {
 	// no-op player count is not interesting to us here
 }
 
-void UPixelStreamingSignallingComponent::FPixelStreamingSignallingConnectionObserver::OnSignallingPlayerConnected(FPixelStreamingPlayerId PlayerId, const FPixelStreamingPlayerConfig& PlayerConfig, bool bSendOffer)
+void UPixelStreamingSignallingComponent::FSignallingObserver::OnSignallingPlayerConnected(FPixelStreamingPlayerId PlayerId, const FPixelStreamingPlayerConfig& PlayerConfig, bool bSendOffer)
 {
 	// no-op our player receiving other players is possible but not a case we are interested in
 	UE_LOG(LogPixelStreamingPlayer, Warning, TEXT("Warning: Received message about player connected from SS, this very likely means you have misconfigured the Pixel Streaming player connection and entered the streamer port. Please check your connect node and ensure it using the player port, e.g. 80 (or no port)."));
 }
 
-void UPixelStreamingSignallingComponent::FPixelStreamingSignallingConnectionObserver::OnSignallingPlayerDisconnected(FPixelStreamingPlayerId PlayerId)
+void UPixelStreamingSignallingComponent::FSignallingObserver::OnSignallingPlayerDisconnected(FPixelStreamingPlayerId PlayerId)
 {
 	// no-op our player knowing when other players disconnect is not interesting for now
 	UE_LOG(LogPixelStreamingPlayer, Warning, TEXT("Warning: Received message about player disconnected from SS, this very likely means you have misconfigured the Pixel Streaming player connection and entered the streamer port. Please check your connect node and ensure it using the player port, e.g. 80 (or no port)."));
 }
 
-void UPixelStreamingSignallingComponent::FPixelStreamingSignallingConnectionObserver::OnSignallingSFUPeerDataChannels(FPixelStreamingPlayerId SFUId, FPixelStreamingPlayerId PlayerId, int32 SendStreamId, int32 RecvStreamId)
+void UPixelStreamingSignallingComponent::FSignallingObserver::OnSignallingSFUPeerDataChannels(FPixelStreamingPlayerId SFUId, FPixelStreamingPlayerId PlayerId, int32 SendStreamId, int32 RecvStreamId)
 {
 	// TODO (Matthew.Cotton): What to do with data channels?
 }
+
+void UPixelStreamingSignallingComponent::FSignallingObserver::OnSignallingStreamerList(const TArray<FString>& StreamerList)
+{
+	Parent->OnStreamerListCallback(StreamerList);
+}
+
+void UPixelStreamingSignallingComponent::AsyncRequestStreamerList(const TFunction<void(const TArray<FString>&)>& Callback)
+{
+	OnStreamerListCallback = Callback;
+	SignallingConnection->RequestStreamerList();
+}
diff --git a/Engine/Plugins/Experimental/PixelStreamingPlayer/Source/PixelStreamingPlayer/Public/PixelStreamingSignallingComponent.h b/Engine/Plugins/Experimental/PixelStreamingPlayer/Source/PixelStreamingPlayer/Public/PixelStreamingSignallingComponent.h
--- a/Engine/Plugins/Experimental/PixelStreamingPlayer/Source/PixelStreamingPlayer/Public/PixelStreamingSignallingComponent.h
+++ b/Engine/Plugins/Experimental/PixelStreamingPlayer/Source/PixelStreamingPlayer/Public/PixelStreamingSignallingComponent.h
@@ -6,6 +6,7 @@
 #include "StreamMediaSource.h"
 #include "Components/ActorComponent.h"
 #include "PixelStreamingWebRTCWrappers.h"
+#include "Kismet/BlueprintAsyncActionBase.h"
 #include "PixelStreamingSignallingComponent.generated.h"
 
 DECLARE_DYNAMIC_MULTICAST_SPARSE_DELEGATE(FPixelStreamingSignallingComponentConnected, UPixelStreamingSignallingComponent, OnConnected);
@@ -40,6 +41,19 @@ public:
 	UFUNCTION(BlueprintCallable, Category = "PixelStreaming")
 	void Disconnect();
 
+	/**
+	 * Send a subscribe message to the signalling server to subscribe to a specific streamer.
+	 * @param StreamerId The name of the streamer to subscribe to.
+	 */
+	UFUNCTION(BlueprintCallable, Category = "PixelStreaming")
+	void Subscribe(const FString& StreamerId);
+
+	/**
+	 * Send an unsubscribe message to the signalling server to unsubscribe from a streamer.
+	 */
+	UFUNCTION(BlueprintCallable, Category = "PixelStreaming")
+	void Unsubscribe();
+
 	/**
 	 * Send an offer created from a Peer Connection to the signalling server.
 	 * @param Offer The answer object created from calling CreateAnswer on a Peer Connection.
@@ -109,11 +123,14 @@ public:
 	UPROPERTY(EditDefaultsOnly, Category = "Properties", META = (DisplayName = "Stream Media Source", AllowPrivateAccess = true))
 	TObjectPtr<UStreamMediaSource> MediaSource;
 
+	// Used by the async bp node
+	void AsyncRequestStreamerList(const TFunction<void(const TArray<FString>&)>& Callback);
+
 protected:
-	class FPixelStreamingSignallingConnectionObserver : public IPixelStreamingSignallingConnectionObserver
+	class FSignallingObserver : public IPixelStreamingSignallingConnectionObserver
 	{
 	public:
-		FPixelStreamingSignallingConnectionObserver(UPixelStreamingSignallingComponent* InParent)
+		FSignallingObserver(UPixelStreamingSignallingComponent* InParent)
 			: Parent(InParent) {}
 
 		//
@@ -132,11 +149,14 @@ protected:
 		virtual void OnSignallingSFUPeerDataChannels(FPixelStreamingPlayerId SFUId, FPixelStreamingPlayerId PlayerId, int32 SendStreamId, int32 RecvStreamId) override;
 		virtual void OnSignallingSessionDescription(FPixelStreamingPlayerId PlayerId, webrtc::SdpType Type, const FString& Sdp) override {}
 		virtual void OnSignallingRemoteIceCandidate(FPixelStreamingPlayerId PlayerId, const FString& SdpMid, int SdpMLineIndex, const FString& Sdp) override {}
+		virtual void OnSignallingStreamerList(const TArray<FString>& StreamerList) override;
 
 	private:
 		TObjectPtr<UPixelStreamingSignallingComponent> Parent;
 	};
 
+	TFunction<void(const TArray<FString>&)> OnStreamerListCallback;
+
 private:
 	TUniquePtr<IPixelStreamingSignallingConnection> SignallingConnection;
 	TSharedPtr<IPixelStreamingSignallingConnectionObserver> SignallingServerConnectionObserver;
diff --git a/Engine/Plugins/Media/PixelStreaming/Source/PixelStreaming/Private/PixelStreamingSignallingConnection.cpp b/Engine/Plugins/Media/PixelStreaming/Source/PixelStreaming/Private/PixelStreamingSignallingConnection.cpp
--- a/Engine/Plugins/Media/PixelStreaming/Source/PixelStreaming/Private/PixelStreamingSignallingConnection.cpp
+++ b/Engine/Plugins/Media/PixelStreaming/Source/PixelStreaming/Private/PixelStreamingSignallingConnection.cpp
@@ -32,6 +32,7 @@ FPixelStreamingSignallingConnection::FPixelStreamingSignallingConnection(TShared
 	RegisterHandler("playerDisconnected", [this](FJsonObjectPtr JsonMsg) { OnPlayerDisconnected(JsonMsg); });
 	RegisterHandler("streamerDataChannels", [this](FJsonObjectPtr JsonMsg) { OnSFUPeerDataChannels(JsonMsg); });
 	RegisterHandler("peerDataChannels", [this](FJsonObjectPtr JsonMsg) { OnPeerDataChannels(JsonMsg); });
+	RegisterHandler("streamerList", [this](FJsonObjectPtr JsonMsg) { OnStreamerList(JsonMsg); });
 }
 
 FPixelStreamingSignallingConnection::~FPixelStreamingSignallingConnection()
@@ -206,6 +207,40 @@ void FPixelStreamingSignallingConnection::SendDisconnectPlayer(FPixelStreamingPl
 	SendMessage(Msg);
 }
 
+void FPixelStreamingSignallingConnection::RequestStreamerList()
+{
+	FJsonObjectPtr Json = MakeShared<FJsonObject>();
+	Json->SetStringField(TEXT("type"), TEXT("listStreamers"));
+
+	const FString JsonStr = UE::PixelStreaming::ToString(Json, false);
+	UE_LOG(LogPixelStreamingSS, Verbose, TEXT("-> SS: listStreamers\n%s"), *JsonStr);
+
+	SendMessage(JsonStr);
+}
+
+void FPixelStreamingSignallingConnection::SendSubscribe(const FString& ToStreamerId)
+{
+	FJsonObjectPtr SubscribeJson = MakeShared<FJsonObject>();
+	SubscribeJson->SetStringField(TEXT("type"), TEXT("subscribe"));
+	SubscribeJson->SetStringField(TEXT("streamerId"), ToStreamerId);
+
+	const FString SubscribeJsonStr = UE::PixelStreaming::ToString(SubscribeJson, false);
+	UE_LOG(LogPixelStreamingSS, Verbose, TEXT("-> SS: subscribe\n%s"), *SubscribeJsonStr);
+
+	SendMessage(SubscribeJsonStr);
+}
+
+void FPixelStreamingSignallingConnection::SendUnsubscribe()
+{
+	FJsonObjectPtr UnsubscribeJson = MakeShared<FJsonObject>();
+	UnsubscribeJson->SetStringField(TEXT("type"), TEXT("unsubscribe"));
+
+	const FString UnsubscribeJsonStr = UE::PixelStreaming::ToString(UnsubscribeJson, false);
+	UE_LOG(LogPixelStreamingSS, Verbose, TEXT("-> SS: unsubscribe\n%s"), *UnsubscribeJsonStr);
+
+	SendMessage(UnsubscribeJsonStr);
+}
+
 void FPixelStreamingSignallingConnection::SendOffer(const webrtc::SessionDescriptionInterface& SDP)
 {
 	FJsonObjectPtr OfferJson = MakeShared<FJsonObject>();
@@ -655,6 +690,20 @@ void FPixelStreamingSignallingConnection::OnPeerDataChannels(const FJsonObjectPt
 	Observer->OnSignallingPeerDataChannels(SendStreamId, RecvStreamId);
 }
 
+void FPixelStreamingSignallingConnection::OnStreamerList(const FJsonObjectPtr& Json)
+{
+	TArray<FString> ResultList;
+	const TArray<TSharedPtr<FJsonValue>>* JsonStreamerIds = nullptr;
+	if (Json->TryGetArrayField("ids", JsonStreamerIds))
+	{
+		for (auto& JsonId : *JsonStreamerIds)
+		{
+			ResultList.Add(JsonId->AsString());
+		}
+	}
+	Observer->OnSignallingStreamerList(ResultList);
+}
+
 void FPixelStreamingSignallingConnection::SetPlayerIdJson(FJsonObjectPtr& JsonObject, FPixelStreamingPlayerId PlayerId)
 {
 	bool bSendAsInteger = UE::PixelStreaming::Settings::CVarSendPlayerIdAsInteger.GetValueOnAnyThread();
diff --git a/Engine/Plugins/Media/PixelStreaming/Source/PixelStreaming/Private/SignallingConnectionObserver.cpp b/Engine/Plugins/Media/PixelStreaming/Source/PixelStreaming/Private/SignallingConnectionObserver.cpp
--- a/Engine/Plugins/Media/PixelStreaming/Source/PixelStreaming/Private/SignallingConnectionObserver.cpp
+++ b/Engine/Plugins/Media/PixelStreaming/Source/PixelStreaming/Private/SignallingConnectionObserver.cpp
@@ -86,7 +86,12 @@ namespace UE::PixelStreaming
 		Streamer.AddNewDataChannel(PlayerId, NewChannel);
 	}
 
-	// Player-only
+	// These are player only and will only be relevant when on the receiving side of pixel streaming, such as the player plugin.
+	void FPixelStreamingSignallingConnectionObserver::OnSignallingStreamerList(const TArray<FString>& StreamerList)
+	{
+		unimplemented();
+	}
+
 	void FPixelStreamingSignallingConnectionObserver::OnSignallingSessionDescription(webrtc::SdpType Type, const FString& Sdp)
 	{
 		unimplemented();
diff --git a/Engine/Plugins/Media/PixelStreaming/Source/PixelStreaming/Private/SignallingConnectionObserver.h b/Engine/Plugins/Media/PixelStreaming/Source/PixelStreaming/Private/SignallingConnectionObserver.h
--- a/Engine/Plugins/Media/PixelStreaming/Source/PixelStreaming/Private/SignallingConnectionObserver.h
+++ b/Engine/Plugins/Media/PixelStreaming/Source/PixelStreaming/Private/SignallingConnectionObserver.h
@@ -26,6 +26,7 @@ namespace UE::PixelStreaming
 		virtual void OnSignallingSFUPeerDataChannels(FPixelStreamingPlayerId SFUId, FPixelStreamingPlayerId PlayerId, int32 SendStreamId, int32 RecvStreamId) override;
 
 		// Player-only
+		virtual void OnSignallingStreamerList(const TArray<FString>& StreamerList) override;
 		virtual void OnSignallingSessionDescription(webrtc::SdpType Type, const FString& Sdp) override;
 		virtual void OnSignallingRemoteIceCandidate(const FString& SdpMid, int SdpMLineIndex, const FString& Sdp) override;
 		virtual void OnSignallingPlayerCount(uint32 Count) override;
diff --git a/Engine/Plugins/Media/PixelStreaming/Source/PixelStreaming/Private/Tests/MockPlayer.h b/Engine/Plugins/Media/PixelStreaming/Source/PixelStreaming/Private/Tests/MockPlayer.h
--- a/Engine/Plugins/Media/PixelStreaming/Source/PixelStreaming/Private/Tests/MockPlayer.h
+++ b/Engine/Plugins/Media/PixelStreaming/Source/PixelStreaming/Private/Tests/MockPlayer.h
@@ -200,6 +200,7 @@ namespace UE::PixelStreaming
 		virtual void OnSignallingPlayerConnected(FPixelStreamingPlayerId PlayerId, const FPixelStreamingPlayerConfig& PlayerConfig, bool bSendOffer) override {}
 		virtual void OnSignallingPlayerDisconnected(FPixelStreamingPlayerId PlayerId) override {}
 		virtual void OnSignallingSFUPeerDataChannels(FPixelStreamingPlayerId SFUId, FPixelStreamingPlayerId PlayerId, int32 SendStreamId, int32 RecvStreamId) override {}
+		virtual void OnSignallingStreamerList(const TArray<FString>& StreamerList) override {}
 
 		DECLARE_MULTICAST_DELEGATE(FOnConnectionEstablished);
 		FOnConnectionEstablished OnConnectionEstablished;
diff --git a/Engine/Plugins/Media/PixelStreaming/Source/PixelStreaming/Private/Tests/SignallingServerConnectionTests.cpp b/Engine/Plugins/Media/PixelStreaming/Source/PixelStreaming/Private/Tests/SignallingServerConnectionTests.cpp
--- a/Engine/Plugins/Media/PixelStreaming/Source/PixelStreaming/Private/Tests/SignallingServerConnectionTests.cpp
+++ b/Engine/Plugins/Media/PixelStreaming/Source/PixelStreaming/Private/Tests/SignallingServerConnectionTests.cpp
@@ -51,6 +51,7 @@ namespace
 		}
 		virtual void OnSignallingPlayerCount(uint32 Count) override {}
 		virtual void OnSignallingPeerDataChannels(int32 SendStreamId, int32 RecvStreamId) override {}
+		virtual void OnSignallingStreamerList(const TArray<FString>& StreamerList) override {}
 
 		void Reset()
 		{
@@ -538,6 +539,40 @@ namespace
 			SendMessage(Msg);
 		}
 
+		virtual void RequestStreamerList() override
+		{
+			FJsonObjectPtr Json = MakeShared<FJsonObject>();
+			Json->SetStringField(TEXT("type"), TEXT("streamerList"));
+
+			const FString JsonStr = UE::PixelStreaming::ToString(Json, false);
+			UE_LOG(LogMockPixelStreamingSS, Verbose, TEXT("-> SS: streamerList\n%s"), *JsonStr);
+
+			SendMessage(JsonStr);
+		}
+
+		virtual void SendSubscribe(const FString& ToStreamerId) override
+		{
+			FJsonObjectPtr SubscribeJson = MakeShared<FJsonObject>();
+			SubscribeJson->SetStringField(TEXT("type"), TEXT("subscribe"));
+			SubscribeJson->SetStringField(TEXT("streamerId"), ToStreamerId);
+
+			const FString SubscribeJsonStr = UE::PixelStreaming::ToString(SubscribeJson, false);
+			UE_LOG(LogMockPixelStreamingSS, Verbose, TEXT("-> SS: subscribe\n%s"), *SubscribeJsonStr);
+
+			SendMessage(SubscribeJsonStr);
+		}
+
+		virtual void SendUnsubscribe() override
+		{
+			FJsonObjectPtr SubscribeJson = MakeShared<FJsonObject>();
+			SubscribeJson->SetStringField(TEXT("type"), TEXT("unsubscribe"));
+
+			const FString SubscribeJsonStr = UE::PixelStreaming::ToString(SubscribeJson, false);
+			UE_LOG(LogMockPixelStreamingSS, Verbose, TEXT("-> SS: subscribe\n%s"), *SubscribeJsonStr);
+
+			SendMessage(SubscribeJsonStr);
+		}
+
 		virtual void SendOffer(const webrtc::SessionDescriptionInterface& SDP) override
 		{
 			FJsonObjectPtr OfferJson = MakeShared<FJsonObject>();
diff --git a/Engine/Plugins/Media/PixelStreaming/Source/PixelStreaming/Public/IPixelStreamingSignallingConnection.h b/Engine/Plugins/Media/PixelStreaming/Source/PixelStreaming/Public/IPixelStreamingSignallingConnection.h
--- a/Engine/Plugins/Media/PixelStreaming/Source/PixelStreaming/Public/IPixelStreamingSignallingConnection.h
+++ b/Engine/Plugins/Media/PixelStreaming/Source/PixelStreaming/Public/IPixelStreamingSignallingConnection.h
@@ -63,19 +63,35 @@ public:
 	virtual void SendDisconnectPlayer(FPixelStreamingPlayerId PlayerId, const FString& Reason) = 0;
 
 	/**
-	 * [Player only] Sends an offer to the streamer.
+	 * Request a list of active streamers from the signalling server. This is a request to the signalling server so can be used by anyone.
+	 */
+	virtual void RequestStreamerList() = 0;
+	
+	/**
+	 * [Stream receiver only] Send a subscription request to the signalling server.
+	 * @param StreamerId The ID of the streamer to subscribe to.
+	 */
+	virtual void SendSubscribe(const FString& StreamerId) = 0;
+
+	/**
+	 * [Stream receiver only] Send an unsubscribe request to the signalling server.
+	 */
+	virtual void SendUnsubscribe() = 0;
+
+	/**
+	 * [Stream receiver only] Sends an offer to the streamer.
 	 * @param SDP The local description offer.
 	 */
 	virtual void SendOffer(const webrtc::SessionDescriptionInterface& SDP) = 0;
 
 	/**
-	 * [Player only] Sends an answer back to the streamer after receiving an offer.
+	 * [Stream receiver only] Sends an answer back to the streamer after receiving an offer.
 	 * @param SDP The local description answer.
 	 */
 	virtual void SendAnswer(const webrtc::SessionDescriptionInterface& SDP) = 0;
 
 	/**
-	 * [Player only] Sends ICE candidate information to the streamer.
+	 * [Stream receiver only] Sends ICE candidate information to the streamer.
 	 * @param IceCandidate The ICE candidate information.
 	 */
 	virtual void SendIceCandidate(const webrtc::IceCandidateInterface& IceCandidate) = 0;
diff --git a/Engine/Plugins/Media/PixelStreaming/Source/PixelStreaming/Public/IPixelStreamingSignallingConnectionObserver.h b/Engine/Plugins/Media/PixelStreaming/Source/PixelStreaming/Public/IPixelStreamingSignallingConnectionObserver.h
--- a/Engine/Plugins/Media/PixelStreaming/Source/PixelStreaming/Public/IPixelStreamingSignallingConnectionObserver.h
+++ b/Engine/Plugins/Media/PixelStreaming/Source/PixelStreaming/Public/IPixelStreamingSignallingConnectionObserver.h
@@ -25,6 +25,7 @@ public:
 	virtual void OnSignallingSFUPeerDataChannels(FPixelStreamingPlayerId SFUId, FPixelStreamingPlayerId PlayerId, int32 SendStreamId, int32 RecvStreamId) = 0;
 
 	// Player-only
+	virtual void OnSignallingStreamerList(const TArray<FString>& StreamerList) = 0;
 	virtual void OnSignallingSessionDescription(webrtc::SdpType Type, const FString& Sdp) = 0;
 	virtual void OnSignallingRemoteIceCandidate(const FString& SdpMid, int SdpMLineIndex, const FString& Sdp) = 0;
 	virtual void OnSignallingPlayerCount(uint32 Count) = 0;
diff --git a/Engine/Plugins/Media/PixelStreaming/Source/PixelStreaming/Public/PixelStreamingSignallingConnection.h b/Engine/Plugins/Media/PixelStreaming/Source/PixelStreaming/Public/PixelStreamingSignallingConnection.h
--- a/Engine/Plugins/Media/PixelStreaming/Source/PixelStreaming/Public/PixelStreamingSignallingConnection.h
+++ b/Engine/Plugins/Media/PixelStreaming/Source/PixelStreaming/Public/PixelStreamingSignallingConnection.h
@@ -28,6 +28,9 @@ public:
 	virtual void SendAnswer(FPixelStreamingPlayerId PlayerId, const webrtc::SessionDescriptionInterface& SDP) override;
 	virtual void SendIceCandidate(FPixelStreamingPlayerId PlayerId, const webrtc::IceCandidateInterface& IceCandidate) override;
 	virtual void SendDisconnectPlayer(FPixelStreamingPlayerId PlayerId, const FString& Reason) override;
+	virtual void RequestStreamerList() override;
+	virtual void SendSubscribe(const FString& StreamerId) override;
+	virtual void SendUnsubscribe() override;
 	virtual void SendOffer(const webrtc::SessionDescriptionInterface& SDP) override;
 	virtual void SendAnswer(const webrtc::SessionDescriptionInterface& SDP) override;
 	virtual void SendIceCandidate(const webrtc::IceCandidateInterface& IceCandidate) override;
@@ -58,6 +61,7 @@ private:
 	void OnPlayerDisconnected(const FJsonObjectPtr& Json);
 	void OnSFUPeerDataChannels(const FJsonObjectPtr& Json);
 	void OnPeerDataChannels(const FJsonObjectPtr& Json);
+	void OnStreamerList(const FJsonObjectPtr& Json);
 	void SetPlayerIdJson(FJsonObjectPtr& JsonObject, FPixelStreamingPlayerId PlayerId);
 	bool GetPlayerIdJson(const FJsonObjectPtr& Json, FPixelStreamingPlayerId& OutPlayerId, const FString& FieldId = TEXT("playerId"));
 
