Do not include PreWindowsApi.h / PostWindowsApi.h directly in code.

 - mostly just replacing them AllowWindowsPlatformTypes.h / HideWindowsPlatformTypes.h
 - some other files had both Pre/PostWindowsApi.h and Allow/HideWindowsPlatformTypes.h, so just removed
 - some occasional include order changes.

#rnx

diff --git a/Engine/Plugins/AI/MLAdapter/Source/MLAdapter/Public/RPCWrapper/rpclib_includes.h b/Engine/Plugins/AI/MLAdapter/Source/MLAdapter/Public/RPCWrapper/rpclib_includes.h
--- a/Engine/Plugins/AI/MLAdapter/Source/MLAdapter/Public/RPCWrapper/rpclib_includes.h
+++ b/Engine/Plugins/AI/MLAdapter/Source/MLAdapter/Public/RPCWrapper/rpclib_includes.h
@@ -24,7 +24,8 @@
 #endif /* ifndef RPCLIB_MSGPACK */
 
 #if PLATFORM_WINDOWS
-#include "Windows/PreWindowsApi.h"
+#include "Windows/AllowWindowsPlatformTypes.h"
+#include "Windows/HideWindowsPlatformTypes.h"
 #endif
 
 #include "rpc/client.h"
@@ -32,10 +33,6 @@
 #include "rpc/this_handler.h"
 #include "rpc/msgpack.hpp"
 
-#if PLATFORM_WINDOWS
-#include "Windows/PostWindowsApi.h"
-#endif
-
 #ifdef __clang__
 #pragma clang diagnostic pop
 #endif
diff --git a/Engine/Plugins/Editor/StylusInput/Source/StylusInput/Private/WindowsRealTimeStylusPlugin.h b/Engine/Plugins/Editor/StylusInput/Source/StylusInput/Private/WindowsRealTimeStylusPlugin.h
--- a/Engine/Plugins/Editor/StylusInput/Source/StylusInput/Private/WindowsRealTimeStylusPlugin.h
+++ b/Engine/Plugins/Editor/StylusInput/Source/StylusInput/Private/WindowsRealTimeStylusPlugin.h
@@ -6,11 +6,9 @@
 
 #if PLATFORM_WINDOWS
 #include "Windows/AllowWindowsPlatformTypes.h"
-#include "Windows/PreWindowsApi.h"
 #include "Microsoft/COMPointer.h"
 	#include <guiddef.h>
 	#include <RTSCom.h>
-#include "Windows/PostWindowsApi.h"
 #include "Windows/HideWindowsPlatformTypes.h"
 
 #include "IStylusState.h"
diff --git a/Engine/Plugins/Experimental/AVCodecs/AMFCodecs/Source/AMFCodecs/Public/AMF.h b/Engine/Plugins/Experimental/AVCodecs/AMFCodecs/Source/AMFCodecs/Public/AMF.h
--- a/Engine/Plugins/Experimental/AVCodecs/AMFCodecs/Source/AMFCodecs/Public/AMF.h
+++ b/Engine/Plugins/Experimental/AVCodecs/AMFCodecs/Source/AMFCodecs/Public/AMF.h
@@ -14,7 +14,6 @@ THIRD_PARTY_INCLUDES_START
 
 #if PLATFORM_WINDOWS
 #include "Windows/AllowWindowsPlatformTypes.h"
-#include "Windows/PreWindowsApi.h"
 #pragma warning(push)
 #pragma warning(disable : 4005)
 #endif //PLATFORM_WINDOWS
@@ -26,7 +25,6 @@ THIRD_PARTY_INCLUDES_START
 
 #if PLATFORM_WINDOWS
 #pragma warning(pop)
-#include "Windows/PostWindowsApi.h"
 #include "Windows/HideWindowsPlatformTypes.h"
 #endif //PLATFORM_WINDOWS
 
diff --git a/Engine/Plugins/Experimental/AVCodecs/AVCodecsCore/Source/AVCodecsCore/Public/Video/Resources/VideoResourceVulkan.h b/Engine/Plugins/Experimental/AVCodecs/AVCodecsCore/Source/AVCodecsCore/Public/Video/Resources/VideoResourceVulkan.h
--- a/Engine/Plugins/Experimental/AVCodecs/AVCodecsCore/Source/AVCodecsCore/Public/Video/Resources/VideoResourceVulkan.h
+++ b/Engine/Plugins/Experimental/AVCodecs/AVCodecsCore/Source/AVCodecsCore/Public/Video/Resources/VideoResourceVulkan.h
@@ -8,13 +8,11 @@
 THIRD_PARTY_INCLUDES_START
 #if PLATFORM_WINDOWS
 #include "Windows/AllowWindowsPlatformTypes.h"
-#include "Windows/PreWindowsApi.h"
 #endif
 
 #include "vulkan_core.h"
 
 #if PLATFORM_WINDOWS
-#include "Windows/PostWindowsApi.h"
 #include "Windows/HideWindowsPlatformTypes.h"
 #endif
 THIRD_PARTY_INCLUDES_END
diff --git a/Engine/Plugins/Experimental/AVCodecs/AVCodecsCore/Source/AVCodecsCore/Public/Video/Resources/Windows/VideoResourceD3D.h b/Engine/Plugins/Experimental/AVCodecs/AVCodecsCore/Source/AVCodecsCore/Public/Video/Resources/Windows/VideoResourceD3D.h
--- a/Engine/Plugins/Experimental/AVCodecs/AVCodecsCore/Source/AVCodecsCore/Public/Video/Resources/Windows/VideoResourceD3D.h
+++ b/Engine/Plugins/Experimental/AVCodecs/AVCodecsCore/Source/AVCodecsCore/Public/Video/Resources/Windows/VideoResourceD3D.h
@@ -12,7 +12,6 @@
 
 THIRD_PARTY_INCLUDES_START
 #include "Windows/AllowWindowsPlatformTypes.h"
-#include "Windows/PreWindowsApi.h"
 #include <d3d11.h>
 #include <mfobjects.h>
 #include <mftransform.h>
@@ -25,7 +24,6 @@ THIRD_PARTY_INCLUDES_START
 #include <d3d11_1.h>
 #include <d3d12.h>
 #include <dxgi1_4.h>
-#include "Windows/PostWindowsApi.h"
 #include "Windows/HideWindowsPlatformTypes.h"
 THIRD_PARTY_INCLUDES_END
 
diff --git a/Engine/Plugins/Experimental/AVCodecs/NVCodecs/Source/NVDEC/Public/NVDEC.h b/Engine/Plugins/Experimental/AVCodecs/NVCodecs/Source/NVDEC/Public/NVDEC.h
--- a/Engine/Plugins/Experimental/AVCodecs/NVCodecs/Source/NVDEC/Public/NVDEC.h
+++ b/Engine/Plugins/Experimental/AVCodecs/NVCodecs/Source/NVDEC/Public/NVDEC.h
@@ -8,14 +8,12 @@ THIRD_PARTY_INCLUDES_START
 
 #if PLATFORM_WINDOWS
 #include "Windows/AllowWindowsPlatformTypes.h"
-#include "Windows/PreWindowsApi.h"
 #endif
 
 #include <cuviddec.h>
 #include <nvcuvid.h>
 
 #if PLATFORM_WINDOWS
-#include "Windows/PostWindowsApi.h"
 #include "Windows/HideWindowsPlatformTypes.h"
 #endif
 
diff --git a/Engine/Plugins/Experimental/AVCodecs/NVCodecs/Source/NVENC/Public/NVENC.h b/Engine/Plugins/Experimental/AVCodecs/NVCodecs/Source/NVENC/Public/NVENC.h
--- a/Engine/Plugins/Experimental/AVCodecs/NVCodecs/Source/NVENC/Public/NVENC.h
+++ b/Engine/Plugins/Experimental/AVCodecs/NVCodecs/Source/NVENC/Public/NVENC.h
@@ -8,13 +8,11 @@ THIRD_PARTY_INCLUDES_START
 
 #if PLATFORM_WINDOWS
 #include "Windows/AllowWindowsPlatformTypes.h"
-#include "Windows/PreWindowsApi.h"
 #endif
 
 #include <nvEncodeAPI.h>
 
 #if PLATFORM_WINDOWS
-#include "Windows/PostWindowsApi.h"
 #include "Windows/HideWindowsPlatformTypes.h"
 #endif
 
diff --git a/Engine/Plugins/Experimental/PixelStreamingPlayer/Source/PixelStreamingPlayer/Private/WebRTCIncludes.h b/Engine/Plugins/Experimental/PixelStreamingPlayer/Source/PixelStreamingPlayer/Private/WebRTCIncludes.h
--- a/Engine/Plugins/Experimental/PixelStreamingPlayer/Source/PixelStreamingPlayer/Private/WebRTCIncludes.h
+++ b/Engine/Plugins/Experimental/PixelStreamingPlayer/Source/PixelStreamingPlayer/Private/WebRTCIncludes.h
@@ -7,7 +7,6 @@ THIRD_PARTY_INCLUDES_START
 #if PLATFORM_WINDOWS
 
 #include "Windows/AllowWindowsPlatformTypes.h"
-#include "Windows/PreWindowsApi.h"
 
 // C4582: constructor is not implicitly called in "api/rtcerror.h", treated as an error by UnrealEngine
 // C6319: Use of the comma-operator in a tested expression causes the left argument to be ignored when it has no side-effects.
@@ -22,7 +21,6 @@ THIRD_PARTY_INCLUDES_START
 #if PLATFORM_WINDOWS
 #pragma warning(pop)
 
-#include "Windows/PostWindowsApi.h"
 #include "Windows/HideWindowsPlatformTypes.h"
 #endif //PLATFORM_WINDOWS
 
diff --git a/Engine/Plugins/Experimental/WebSocketNetworking/Source/WebSocketNetworking/Private/WebSocketServer.cpp b/Engine/Plugins/Experimental/WebSocketNetworking/Source/WebSocketNetworking/Private/WebSocketServer.cpp
--- a/Engine/Plugins/Experimental/WebSocketNetworking/Source/WebSocketNetworking/Private/WebSocketServer.cpp
+++ b/Engine/Plugins/Experimental/WebSocketNetworking/Source/WebSocketNetworking/Private/WebSocketServer.cpp
@@ -8,9 +8,7 @@
 
 #if PLATFORM_WINDOWS
 #include "Windows/AllowWindowsPlatformTypes.h"
-#include "Windows/PreWindowsApi.h"
 #include <winsock2.h>
-#include "Windows/PostWindowsApi.h"
 #include "Windows/HideWindowsPlatformTypes.h"
 #endif
 
diff --git a/Engine/Plugins/Media/AjaMedia/Source/Aja/Private/Helpers.h b/Engine/Plugins/Media/AjaMedia/Source/Aja/Private/Helpers.h
--- a/Engine/Plugins/Media/AjaMedia/Source/Aja/Private/Helpers.h
+++ b/Engine/Plugins/Media/AjaMedia/Source/Aja/Private/Helpers.h
@@ -16,7 +16,6 @@ __pragma(warning(disable: 4264))  /* No override available for virtual member fu
 
 #if PLATFORM_WINDOWS
 #include "Windows/AllowWindowsPlatformTypes.h"
-#include "Windows/PreWindowsApi.h"
 #pragma warning(push)
 #pragma warning(disable : 4005)
 #endif //PLATFORM_WINDOWS
@@ -41,7 +40,6 @@ __pragma(warning(disable: 4264))  /* No override available for virtual member fu
 
 #if PLATFORM_WINDOWS
 #pragma warning(pop)
-#include "Windows/PostWindowsApi.h"
 #include "Windows/HideWindowsPlatformTypes.h"
 #endif //PLATFORM_WINDOWS
 
diff --git a/Engine/Plugins/Media/HardwareEncoders/Source/EncoderAMF/Private/Amf_Common.h b/Engine/Plugins/Media/HardwareEncoders/Source/EncoderAMF/Private/Amf_Common.h
--- a/Engine/Plugins/Media/HardwareEncoders/Source/EncoderAMF/Private/Amf_Common.h
+++ b/Engine/Plugins/Media/HardwareEncoders/Source/EncoderAMF/Private/Amf_Common.h
@@ -8,7 +8,6 @@ THIRD_PARTY_INCLUDES_START
 
 #if PLATFORM_WINDOWS
 #include "Windows/AllowWindowsPlatformTypes.h"
-#include "Windows/PreWindowsApi.h"
 #pragma warning(push)
 #pragma warning(disable : 4005)
 #endif //PLATFORM_WINDOWS
@@ -21,7 +20,6 @@ THIRD_PARTY_INCLUDES_START
 
 #if PLATFORM_WINDOWS
 #pragma warning(pop)
-#include "Windows/PostWindowsApi.h"
 #include "Windows/HideWindowsPlatformTypes.h"
 #endif //PLATFORM_WINDOWS
 
diff --git a/Engine/Plugins/Media/HardwareEncoders/Source/EncoderNVENC/Private/NVENC_Common.h b/Engine/Plugins/Media/HardwareEncoders/Source/EncoderNVENC/Private/NVENC_Common.h
--- a/Engine/Plugins/Media/HardwareEncoders/Source/EncoderNVENC/Private/NVENC_Common.h
+++ b/Engine/Plugins/Media/HardwareEncoders/Source/EncoderNVENC/Private/NVENC_Common.h
@@ -8,13 +8,11 @@ THIRD_PARTY_INCLUDES_START
 
 #if PLATFORM_WINDOWS
 #include "Windows/AllowWindowsPlatformTypes.h"
-#include "Windows/PreWindowsApi.h"
 #endif
 
 #include <nvEncodeAPI.h>
 
 #if PLATFORM_WINDOWS
-#include "Windows/PostWindowsApi.h"
 #include "Windows/HideWindowsPlatformTypes.h"
 #endif
 
diff --git a/Engine/Plugins/Media/MediaIOFramework/Source/GPUTextureTransfer/Private/D3D11TextureTransfer.cpp b/Engine/Plugins/Media/MediaIOFramework/Source/GPUTextureTransfer/Private/D3D11TextureTransfer.cpp
--- a/Engine/Plugins/Media/MediaIOFramework/Source/GPUTextureTransfer/Private/D3D11TextureTransfer.cpp
+++ b/Engine/Plugins/Media/MediaIOFramework/Source/GPUTextureTransfer/Private/D3D11TextureTransfer.cpp
@@ -4,11 +4,9 @@
 
 #if DVP_SUPPORTED_PLATFORM
 #include "Windows/AllowWindowsPlatformTypes.h"
-#include "Windows/PreWindowsApi.h"
 
 #include "dvpapi_d3d11.h"
 
-#include "Windows/PostWindowsApi.h"
 #include "Windows/HideWindowsPlatformTypes.h"
 
 namespace UE::GPUTextureTransfer::Private
diff --git a/Engine/Plugins/Media/MediaIOFramework/Source/GPUTextureTransfer/Private/D3D12TextureTransfer.cpp b/Engine/Plugins/Media/MediaIOFramework/Source/GPUTextureTransfer/Private/D3D12TextureTransfer.cpp
--- a/Engine/Plugins/Media/MediaIOFramework/Source/GPUTextureTransfer/Private/D3D12TextureTransfer.cpp
+++ b/Engine/Plugins/Media/MediaIOFramework/Source/GPUTextureTransfer/Private/D3D12TextureTransfer.cpp
@@ -4,11 +4,9 @@
 
 #if DVP_SUPPORTED_PLATFORM
 #include "Windows/AllowWindowsPlatformTypes.h"
-#include "Windows/PreWindowsApi.h"
 
 #include "dvpapi_d3d12.h"
 
-#include "Windows/PostWindowsApi.h"
 #include "Windows/HideWindowsPlatformTypes.h"
 
 
diff --git a/Engine/Plugins/Media/MediaIOFramework/Source/GPUTextureTransfer/Private/TextureTransferBase.h b/Engine/Plugins/Media/MediaIOFramework/Source/GPUTextureTransfer/Private/TextureTransferBase.h
--- a/Engine/Plugins/Media/MediaIOFramework/Source/GPUTextureTransfer/Private/TextureTransferBase.h
+++ b/Engine/Plugins/Media/MediaIOFramework/Source/GPUTextureTransfer/Private/TextureTransferBase.h
@@ -6,9 +6,7 @@
 
 #if DVP_SUPPORTED_PLATFORM
 #include "Windows/AllowWindowsPlatformTypes.h"
-#include "Windows/PreWindowsApi.h"
 #include "DVPAPI.h"
-#include "Windows/PostWindowsApi.h"
 #include "Windows/HideWindowsPlatformTypes.h"
 
 class FRHITexture;
diff --git a/Engine/Plugins/Media/MediaIOFramework/Source/GPUTextureTransfer/Private/VulkanTextureTransfer.cpp b/Engine/Plugins/Media/MediaIOFramework/Source/GPUTextureTransfer/Private/VulkanTextureTransfer.cpp
--- a/Engine/Plugins/Media/MediaIOFramework/Source/GPUTextureTransfer/Private/VulkanTextureTransfer.cpp
+++ b/Engine/Plugins/Media/MediaIOFramework/Source/GPUTextureTransfer/Private/VulkanTextureTransfer.cpp
@@ -4,11 +4,9 @@
 
 #if DVP_SUPPORTED_PLATFORM
 #include "Windows/AllowWindowsPlatformTypes.h"
-#include "Windows/PreWindowsApi.h"
 
 #include "dvpapi_vulkan.h"
 
-#include "Windows/PostWindowsApi.h"
 #include "Windows/HideWindowsPlatformTypes.h"
 
 
diff --git a/Engine/Plugins/Media/PixelStreaming/Source/PixelStreaming/Private/WebRTCIncludes.h b/Engine/Plugins/Media/PixelStreaming/Source/PixelStreaming/Private/WebRTCIncludes.h
--- a/Engine/Plugins/Media/PixelStreaming/Source/PixelStreaming/Private/WebRTCIncludes.h
+++ b/Engine/Plugins/Media/PixelStreaming/Source/PixelStreaming/Private/WebRTCIncludes.h
@@ -10,7 +10,6 @@ THIRD_PARTY_INCLUDES_START
 #if PLATFORM_WINDOWS
 
 	#include "Windows/AllowWindowsPlatformTypes.h"
-	#include "Windows/PreWindowsApi.h"
 
 	// C4582: constructor is not implicitly called in "api/rtcerror.h", treated as an error by UnrealEngine
 	// C6319: Use of the comma-operator in a tested expression causes the left argument to be ignored when it has no side-effects.
@@ -89,7 +88,6 @@ THIRD_PARTY_INCLUDES_START
 #if PLATFORM_WINDOWS
 	#pragma warning(pop)
 
-	#include "Windows/PostWindowsApi.h"
 	#include "Windows/HideWindowsPlatformTypes.h"
 #endif // PLATFORM_WINDOWS
 
diff --git a/Engine/Plugins/Media/PixelStreaming/Source/PixelStreaming/Public/PixelStreamingWebRTCIncludes.h b/Engine/Plugins/Media/PixelStreaming/Source/PixelStreaming/Public/PixelStreamingWebRTCIncludes.h
--- a/Engine/Plugins/Media/PixelStreaming/Source/PixelStreaming/Public/PixelStreamingWebRTCIncludes.h
+++ b/Engine/Plugins/Media/PixelStreaming/Source/PixelStreaming/Public/PixelStreamingWebRTCIncludes.h
@@ -8,7 +8,6 @@ THIRD_PARTY_INCLUDES_START
 
 #if PLATFORM_WINDOWS
 	#include "Windows/AllowWindowsPlatformTypes.h"
-	#include "Windows/PreWindowsApi.h"
 
 	// C4582: constructor is not implicitly called in "api/rtcerror.h", treated as an error by UnrealEngine
 	// C6319: Use of the comma-operator in a tested expression causes the left argument to be ignored when it has no side-effects.
@@ -36,7 +35,6 @@ THIRD_PARTY_INCLUDES_START
 #if PLATFORM_WINDOWS
 	#pragma warning(pop)
 
-	#include "Windows/PostWindowsApi.h"
 	#include "Windows/HideWindowsPlatformTypes.h"
 #else
 	#ifdef PF_MAX
diff --git a/Engine/Plugins/Media/WmfMedia/Source/WmfMedia/Public/WmfMediaCommon.h b/Engine/Plugins/Media/WmfMedia/Source/WmfMedia/Public/WmfMediaCommon.h
--- a/Engine/Plugins/Media/WmfMedia/Source/WmfMedia/Public/WmfMediaCommon.h
+++ b/Engine/Plugins/Media/WmfMedia/Source/WmfMedia/Public/WmfMediaCommon.h
@@ -10,7 +10,6 @@
 	#include "WmfMediaSettings.h"
 
 	#include "Windows/AllowWindowsPlatformTypes.h"
-	#include "Windows/PreWindowsApi.h"
 
 	THIRD_PARTY_INCLUDES_START
 
@@ -48,7 +47,6 @@
 
 	#include "Microsoft/COMPointer.h"
 	
-	#include "Windows/PostWindowsApi.h"
 	#include "Windows/HideWindowsPlatformTypes.h"
 
 #elif PLATFORM_WINDOWS && !UE_SERVER
diff --git a/Engine/Plugins/Runtime/nDisplay/Source/DisplayClusterProjection/Private/Policy/VIOSO/DisplayClusterProjectionVIOSOTypes.h b/Engine/Plugins/Runtime/nDisplay/Source/DisplayClusterProjection/Private/Policy/VIOSO/DisplayClusterProjectionVIOSOTypes.h
--- a/Engine/Plugins/Runtime/nDisplay/Source/DisplayClusterProjection/Private/Policy/VIOSO/DisplayClusterProjectionVIOSOTypes.h
+++ b/Engine/Plugins/Runtime/nDisplay/Source/DisplayClusterProjection/Private/Policy/VIOSO/DisplayClusterProjectionVIOSOTypes.h
@@ -10,9 +10,7 @@
 
 THIRD_PARTY_INCLUDES_START
 #include "Windows/AllowWindowsPlatformTypes.h"
-#include "Windows/PreWindowsApi.h"
 #include "VWBTypes.h"
-#include "Windows/PostWindowsApi.h"
 #include "Windows/HideWindowsPlatformTypes.h"
 THIRD_PARTY_INCLUDES_END
 
diff --git a/Engine/Plugins/VirtualProduction/Rivermax/RivermaxCore/Source/RivermaxCore/Private/RivermaxDeviceFinder.cpp b/Engine/Plugins/VirtualProduction/Rivermax/RivermaxCore/Source/RivermaxCore/Private/RivermaxDeviceFinder.cpp
--- a/Engine/Plugins/VirtualProduction/Rivermax/RivermaxCore/Source/RivermaxCore/Private/RivermaxDeviceFinder.cpp
+++ b/Engine/Plugins/VirtualProduction/Rivermax/RivermaxCore/Source/RivermaxCore/Private/RivermaxDeviceFinder.cpp
@@ -9,13 +9,11 @@
 
 #if PLATFORM_WINDOWS
 #include "Windows/AllowWindowsPlatformTypes.h"
-#include "Windows/PreWindowsApi.h"
 
 #include <winsock2.h>
 #include <Iphlpapi.h>
 #include <ws2tcpip.h>
 
-#include "Windows/PostWindowsApi.h"
 #include "Windows/HideWindowsPlatformTypes.h"
 #endif
 
diff --git a/Engine/Plugins/VirtualProduction/Rivermax/RivermaxCore/Source/RivermaxCore/Private/RivermaxManager.cpp b/Engine/Plugins/VirtualProduction/Rivermax/RivermaxCore/Source/RivermaxCore/Private/RivermaxManager.cpp
--- a/Engine/Plugins/VirtualProduction/Rivermax/RivermaxCore/Source/RivermaxCore/Private/RivermaxManager.cpp
+++ b/Engine/Plugins/VirtualProduction/Rivermax/RivermaxCore/Source/RivermaxCore/Private/RivermaxManager.cpp
@@ -16,11 +16,9 @@
 
 #if PLATFORM_WINDOWS
 #include "Windows/AllowWindowsPlatformTypes.h"
-#include "Windows/PreWindowsApi.h"
 
 #include <ws2tcpip.h>
 
-#include "Windows/PostWindowsApi.h"
 #include "Windows/HideWindowsPlatformTypes.h"
 #endif
 
diff --git a/Engine/Plugins/VirtualProduction/Rivermax/RivermaxCore/Source/RivermaxCore/Public/RivermaxHeader.h b/Engine/Plugins/VirtualProduction/Rivermax/RivermaxCore/Source/RivermaxCore/Public/RivermaxHeader.h
--- a/Engine/Plugins/VirtualProduction/Rivermax/RivermaxCore/Source/RivermaxCore/Public/RivermaxHeader.h
+++ b/Engine/Plugins/VirtualProduction/Rivermax/RivermaxCore/Source/RivermaxCore/Public/RivermaxHeader.h
@@ -6,13 +6,11 @@ THIRD_PARTY_INCLUDES_START
 
 #if PLATFORM_WINDOWS
 #include "Windows/AllowWindowsPlatformTypes.h"
-#include "Windows/PreWindowsApi.h"
 #endif
 
 #include "rivermax_api.h"
 
 #if PLATFORM_WINDOWS
-#include "Windows/PostWindowsApi.h"
 #include "Windows/HideWindowsPlatformTypes.h"
 #endif
 
diff --git a/Engine/Plugins/VirtualProduction/TextureShare/Source/TextureShareCore/Private/Core/ResourcesCache/TextureShareCoreResourcesCacheItem.cpp b/Engine/Plugins/VirtualProduction/TextureShare/Source/TextureShareCore/Private/Core/ResourcesCache/TextureShareCoreResourcesCacheItem.cpp
--- a/Engine/Plugins/VirtualProduction/TextureShare/Source/TextureShareCore/Private/Core/ResourcesCache/TextureShareCoreResourcesCacheItem.cpp
+++ b/Engine/Plugins/VirtualProduction/TextureShare/Source/TextureShareCore/Private/Core/ResourcesCache/TextureShareCoreResourcesCacheItem.cpp
@@ -7,9 +7,7 @@
 #if PLATFORM_WINDOWS
 THIRD_PARTY_INCLUDES_START
 #include "Windows/AllowWindowsPlatformTypes.h"
-#include "Windows/PreWindowsApi.h"
 #include <handleapi.h>
-#include "Windows/PostWindowsApi.h"
 #include "Windows/HideWindowsPlatformTypes.h"
 THIRD_PARTY_INCLUDES_END
 #endif
diff --git a/Engine/Plugins/VirtualProduction/TextureShare/Source/TextureShareCore/Private/Core/ResourcesCache/Vulkan/TextureShareCoreVulkanResourcesCacheItem.cpp b/Engine/Plugins/VirtualProduction/TextureShare/Source/TextureShareCore/Private/Core/ResourcesCache/Vulkan/TextureShareCoreVulkanResourcesCacheItem.cpp
--- a/Engine/Plugins/VirtualProduction/TextureShare/Source/TextureShareCore/Private/Core/ResourcesCache/Vulkan/TextureShareCoreVulkanResourcesCacheItem.cpp
+++ b/Engine/Plugins/VirtualProduction/TextureShare/Source/TextureShareCore/Private/Core/ResourcesCache/Vulkan/TextureShareCoreVulkanResourcesCacheItem.cpp
@@ -8,12 +8,10 @@
 #if PLATFORM_WINDOWS
 THIRD_PARTY_INCLUDES_START
 #include "Windows/AllowWindowsPlatformTypes.h"
-#include "Windows/PreWindowsApi.h"
 
 #include "vulkan/vulkan_core.h"
 #include "vulkan/vulkan_win32.h"
 
-#include "Windows/PostWindowsApi.h"
 #include "Windows/HideWindowsPlatformTypes.h"
 THIRD_PARTY_INCLUDES_END
 #endif
diff --git a/Engine/Plugins/VirtualProduction/TextureShare/Source/TextureShareCore/Private/Core/TextureShareCoreSecurityAttributes.cpp b/Engine/Plugins/VirtualProduction/TextureShare/Source/TextureShareCore/Private/Core/TextureShareCoreSecurityAttributes.cpp
--- a/Engine/Plugins/VirtualProduction/TextureShare/Source/TextureShareCore/Private/Core/TextureShareCoreSecurityAttributes.cpp
+++ b/Engine/Plugins/VirtualProduction/TextureShare/Source/TextureShareCore/Private/Core/TextureShareCoreSecurityAttributes.cpp
@@ -9,9 +9,7 @@
 #if PLATFORM_WINDOWS
 THIRD_PARTY_INCLUDES_START
 #include "Windows/AllowWindowsPlatformTypes.h"
-#include "Windows/PreWindowsApi.h"
 #include <aclapi.h>
-#include "Windows/PostWindowsApi.h"
 #include "Windows/HideWindowsPlatformTypes.h"
 THIRD_PARTY_INCLUDES_END
 #endif
diff --git a/Engine/Source/Developer/Apple/MetalShaderFormat/Private/MetalShaderCompiler.cpp b/Engine/Source/Developer/Apple/MetalShaderFormat/Private/MetalShaderCompiler.cpp
--- a/Engine/Source/Developer/Apple/MetalShaderFormat/Private/MetalShaderCompiler.cpp
+++ b/Engine/Source/Developer/Apple/MetalShaderFormat/Private/MetalShaderCompiler.cpp
@@ -19,11 +19,9 @@
 #if PLATFORM_WINDOWS
 #include "Windows/AllowWindowsPlatformTypes.h"
 THIRD_PARTY_INCLUDES_START
-	#include "Windows/PreWindowsApi.h"
 	#include <objbase.h>
 	#include <assert.h>
 	#include <stdio.h>
-	#include "Windows/PostWindowsApi.h"
 THIRD_PARTY_INCLUDES_END
 #include "Windows/HideWindowsPlatformTypes.h"
 #endif
diff --git a/Engine/Source/Developer/ShaderFormatOpenGL/Private/OpenGLShaderCompiler.cpp b/Engine/Source/Developer/ShaderFormatOpenGL/Private/OpenGLShaderCompiler.cpp
--- a/Engine/Source/Developer/ShaderFormatOpenGL/Private/OpenGLShaderCompiler.cpp
+++ b/Engine/Source/Developer/ShaderFormatOpenGL/Private/OpenGLShaderCompiler.cpp
@@ -17,11 +17,9 @@
 
 #if PLATFORM_WINDOWS
 #include "Windows/AllowWindowsPlatformTypes.h"
-	#include "Windows/PreWindowsApi.h"
 	#include <objbase.h>
 	#include <assert.h>
 	#include <stdio.h>
-	#include "Windows/PostWindowsApi.h"
 #include "Windows/HideWindowsPlatformTypes.h"
 #endif
 #include "ShaderCore.h"
diff --git a/Engine/Source/Programs/ZenDashboard/Private/ZenDashboard.cpp b/Engine/Source/Programs/ZenDashboard/Private/ZenDashboard.cpp
--- a/Engine/Source/Programs/ZenDashboard/Private/ZenDashboard.cpp
+++ b/Engine/Source/Programs/ZenDashboard/Private/ZenDashboard.cpp
@@ -29,13 +29,11 @@
 #include "ZenDashboardStyle.h"
 
 #if PLATFORM_WINDOWS
-#include "Windows/WindowsHWrapper.h"
 #include "Runtime/Launch/Resources/Windows/Resource.h"
 #include "Windows/WindowsApplication.h"
-#include "Windows/WindowsHWrapper.h"
-#include "Windows/PreWindowsApi.h"
+#include "Windows/AllowWindowsPlatformTypes.h"
 #include <shellapi.h>
-#include "Windows/PostWindowsApi.h"
+#include "Windows/HideWindowsPlatformTypes.h"
 #elif PLATFORM_LINUX
 #include "UnixCommonStartup.h"
 #elif PLATFORM_MAC
diff --git a/Engine/Source/Runtime/AVEncoder/Public/VideoCommon.h b/Engine/Source/Runtime/AVEncoder/Public/VideoCommon.h
--- a/Engine/Source/Runtime/AVEncoder/Public/VideoCommon.h
+++ b/Engine/Source/Runtime/AVEncoder/Public/VideoCommon.h
@@ -18,7 +18,6 @@
 
 THIRD_PARTY_INCLUDES_START
 #include "Microsoft/AllowMicrosoftPlatformTypes.h"
-#include "Microsoft/PreWindowsApi.h"
 #include <d3d11.h>
 #include <mfobjects.h>
 #include <mftransform.h>
@@ -31,7 +30,6 @@ THIRD_PARTY_INCLUDES_START
 #include <d3d11_1.h>
 #include <d3d12.h>
 #include <dxgi1_4.h>
-#include "Microsoft/PostWindowsApi.h"
 #include "Microsoft/HideMicrosoftPlatformTypes.h"
 THIRD_PARTY_INCLUDES_END
 
diff --git a/Engine/Source/Runtime/ApplicationCore/Private/Windows/WindowsPlatformApplicationMisc.cpp b/Engine/Source/Runtime/ApplicationCore/Private/Windows/WindowsPlatformApplicationMisc.cpp
--- a/Engine/Source/Runtime/ApplicationCore/Private/Windows/WindowsPlatformApplicationMisc.cpp
+++ b/Engine/Source/Runtime/ApplicationCore/Private/Windows/WindowsPlatformApplicationMisc.cpp
@@ -22,11 +22,9 @@
 
 THIRD_PARTY_INCLUDES_START
 #include "Windows/AllowWindowsPlatformTypes.h"
-#include "Windows/PreWindowsApi.h"
 #include "dxgi1_3.h"
 #include "dxgi1_4.h"
 #include "dxgi1_6.h"
-#include "Windows/PostWindowsApi.h"
 #include "Windows/HideWindowsPlatformTypes.h"
 THIRD_PARTY_INCLUDES_END
 
diff --git a/Engine/Source/Runtime/Core/Private/ProfilingDebugging/Windows/ConcurrencyViewerExternalProfiler.cpp b/Engine/Source/Runtime/Core/Private/ProfilingDebugging/Windows/ConcurrencyViewerExternalProfiler.cpp
--- a/Engine/Source/Runtime/Core/Private/ProfilingDebugging/Windows/ConcurrencyViewerExternalProfiler.cpp
+++ b/Engine/Source/Runtime/Core/Private/ProfilingDebugging/Windows/ConcurrencyViewerExternalProfiler.cpp
@@ -13,9 +13,7 @@
 
 #include "Windows/AllowWindowsPlatformTypes.h"
 #include "Windows/AllowWindowsPlatformAtomics.h"
-#include "Windows/PreWindowsApi.h"
 #include <cvmarkers.h>
-#include "Windows/PostWindowsApi.h"
 #include "Windows/HideWindowsPlatformAtomics.h"
 #include "Windows/HideWindowsPlatformTypes.h"
 
