Fix: Restore PixelStreamingPlayer

Brings the player functionality from 5.1 to Main, including fixing deprecations and adding overrides for required methods.

diff --git a/Engine/Plugins/Experimental/PixelStreamingPlayer/PixelStreamingPlayer.uplugin b/Engine/Plugins/Experimental/PixelStreamingPlayer/PixelStreamingPlayer.uplugin
--- a/Engine/Plugins/Experimental/PixelStreamingPlayer/PixelStreamingPlayer.uplugin
+++ b/Engine/Plugins/Experimental/PixelStreamingPlayer/PixelStreamingPlayer.uplugin
@@ -19,14 +19,14 @@
 			"Name": "PixelStreamingPlayer",
 			"Type": "Runtime",
 			"LoadingPhase": "PreDefault",
-			"PlatformAllowList": [ "Win64" ],
+			"PlatformAllowList": [ "Win64", "Linux" ],
 			"TargetDenyList": [ "Server" ]
 		},
 		{
 			"Name": "PixelStreamingPlayerEditor",
 			"Type": "Editor",
 			"LoadingPhase": "PostEngineInit",
-			"PlatformAllowList": [ "Win64" ],
+			"PlatformAllowList": [ "Win64", "Linux" ],
 			"TargetDenyList": [ "Server" ]
 		}
 	],
@@ -34,10 +34,6 @@
 		{
 			"Name": "PixelStreaming",
 			"Enabled": true
-		},
-		{
-			"Name": "TextureMediaPlayer",
-			"Enabled": true
 		}
 	]
 }
diff --git a/Engine/Plugins/Experimental/PixelStreamingPlayer/Source/PixelStreamingPlayer/PixelStreamingPlayer.Build.cs b/Engine/Plugins/Experimental/PixelStreamingPlayer/Source/PixelStreamingPlayer/PixelStreamingPlayer.Build.cs
--- a/Engine/Plugins/Experimental/PixelStreamingPlayer/Source/PixelStreamingPlayer/PixelStreamingPlayer.Build.cs
+++ b/Engine/Plugins/Experimental/PixelStreamingPlayer/Source/PixelStreamingPlayer/PixelStreamingPlayer.Build.cs
@@ -34,13 +34,9 @@ namespace UnrealBuildTool.Rules
 					"Engine",
 					"MediaAssets",
 					"PixelStreaming",
-					"Renderer",
 					"RenderCore",
 					"RHI",
-					"TextureMediaPlayer",
 					"WebSockets",
-					"AVCodecsCore",
-					"AVCodecsCoreRHI"
 				});
 		}
 	}
diff --git a/Engine/Plugins/Experimental/PixelStreamingPlayer/Source/PixelStreamingPlayer/Private/PixelStreamingMediaTexture.cpp b/Engine/Plugins/Experimental/PixelStreamingPlayer/Source/PixelStreamingPlayer/Private/PixelStreamingMediaTexture.cpp
new file mode 100644
--- /dev/null
+++ b/Engine/Plugins/Experimental/PixelStreamingPlayer/Source/PixelStreamingPlayer/Private/PixelStreamingMediaTexture.cpp
@@ -0,0 +1,166 @@
+// Copyright Epic Games, Inc. All Rights Reserved.
+
+#include "PixelStreamingMediaTexture.h"
+#include "PixelStreamingMediaTextureResource.h"
+#include "Async/Async.h"
+#include "WebRTCIncludes.h"
+
+constexpr int32 DEFAULT_WIDTH = 1920;
+constexpr int32 DEFAULT_HEIGHT = 1080;
+
+UPixelStreamingMediaTexture::UPixelStreamingMediaTexture(const FObjectInitializer& ObjectInitializer)
+	: Super(ObjectInitializer)
+{
+	SetResource(nullptr);
+}
+
+void UPixelStreamingMediaTexture::BeginDestroy()
+{
+	AsyncTask(ENamedThreads::ActualRenderingThread, [this]() {
+		FScopeLock Lock(&RenderSyncContext);
+		RenderTarget = nullptr;
+	});
+
+	Super::BeginDestroy();
+}
+
+void UPixelStreamingMediaTexture::GetResourceSizeEx(FResourceSizeEx& CumulativeResourceSize)
+{
+	Super::GetResourceSizeEx(CumulativeResourceSize);
+
+	if (CurrentResource != nullptr)
+	{
+		CumulativeResourceSize.AddUnknownMemoryBytes(CurrentResource->GetResourceSize());
+	}
+}
+
+FTextureResource* UPixelStreamingMediaTexture::CreateResource()
+{
+	if (CurrentResource != nullptr)
+	{
+		SetResource(nullptr);
+		delete CurrentResource;
+	}
+
+	CurrentResource = new FPixelStreamingMediaTextureResource(this);
+	InitializeResources();
+
+	return CurrentResource;
+}
+
+void UPixelStreamingMediaTexture::OnFrame(const webrtc::VideoFrame& frame)
+{
+	// TODO
+	// Currently this will cause all input frames to be converted to I420 (using webrtc::ConvertFromI420)
+	// and then render the rgba buffer to a texture, but in some cases we might already have the native
+	// texture and can use it or copy it directly. At some point it would be useful to detect and fork this
+	// behaviour.
+	// Additionally, we're rendering to a texture then updating the texture resource with every frame. We
+	// might be able to just map a single dest texture and update it each frame.
+
+	static constexpr auto TEXTURE_PIXEL_FORMAT = PF_B8G8R8A8;
+	static constexpr auto WEBRTC_PIXEL_FORMAT = webrtc::VideoType::kARGB;
+
+	SizeX = frame.width();
+	SizeY = frame.height();
+	Format = TEXTURE_PIXEL_FORMAT;
+	NumMips = 0;
+	SamplerAddressMode = ESamplerAddressMode::AM_Clamp;
+
+	{
+		FScopeLock Lock(&RenderSyncContext);
+
+		const uint32_t Size = webrtc::CalcBufferSize(WEBRTC_PIXEL_FORMAT, SizeX, SizeY);
+		if (Size > static_cast<uint32_t>(Buffer.Num()))
+		{
+			Buffer.SetNum(Size);
+		}
+		webrtc::ConvertFromI420(frame, WEBRTC_PIXEL_FORMAT, 0, Buffer.GetData());
+	}
+
+	AsyncTask(ENamedThreads::ActualRenderingThread, [this]() {
+		FScopeLock Lock(&RenderSyncContext);
+
+		const FIntPoint FrameSize = FIntPoint(SizeX, SizeY);
+		FRHICommandListImmediate& RHICmdList = FRHICommandListExecutor::GetImmediateCommandList();
+
+		if (!RenderTargetDescriptor.IsValid() || RenderTargetDescriptor.GetSize() != FIntVector(FrameSize.X, FrameSize.Y, 0))
+		{
+			// Create the RenderTarget descriptor
+			RenderTargetDescriptor = FPooledRenderTargetDesc::Create2DDesc(FrameSize,
+				Format,
+				FClearValueBinding::None,
+				TexCreate_None,
+				TexCreate_RenderTargetable,
+				false);
+
+			// Update the shader resource for the 'SourceTexture'
+			FRHITextureCreateDesc RenderTargetTextureDesc =
+				FRHITextureCreateDesc::Create2D(TEXT(""), FrameSize.X, FrameSize.Y, Format)
+					.SetClearValue(FClearValueBinding(FLinearColor(0.0f, 0.0f, 0.0f)))
+					.SetFlags(ETextureCreateFlags::Dynamic | ETextureCreateFlags::ShaderResource | TexCreate_RenderTargetable | ETextureCreateFlags::SRGB)
+					.SetInitialState(ERHIAccess::SRVMask);
+
+			SourceTexture = RHICreateTexture(RenderTargetTextureDesc);
+
+			// Find a free target-able texture from the render pool
+			GRenderTargetPool.FindFreeElement(RHICmdList,
+				RenderTargetDescriptor,
+				RenderTarget,
+				TEXT("PIXELSTEAMINGPLAYER"));
+		}
+
+		// Create the update region structure
+		const FUpdateTextureRegion2D Region(0, 0, 0, 0, FrameSize.X, FrameSize.Y);
+
+		// Set the Pixel data of the webrtc Frame to the SourceTexture
+		RHIUpdateTexture2D(SourceTexture, 0, Region, FrameSize.X * 4, Buffer.GetData());
+
+		UpdateTextureReference(RHICmdList, SourceTexture);
+	});
+}
+
+void UPixelStreamingMediaTexture::InitializeResources()
+{
+	ENQUEUE_RENDER_COMMAND(FPixelStreamingMediaTextureUpdateTextureReference)
+	([this](FRHICommandListImmediate& RHICmdList) {
+		// Set the default video texture to reference nothing
+		FTextureRHIRef ShaderTexture2D;
+		FTextureRHIRef RenderableTexture;
+	
+		FRHITextureCreateDesc RenderTargetTextureDesc =
+			FRHITextureCreateDesc::Create2D(TEXT(""), DEFAULT_WIDTH, DEFAULT_HEIGHT, EPixelFormat::PF_B8G8R8A8)
+				.SetClearValue(FClearValueBinding(FLinearColor(0.0f, 0.0f, 0.0f)))
+				.SetFlags(ETextureCreateFlags::Dynamic | ETextureCreateFlags::ShaderResource | TexCreate_RenderTargetable)
+				.SetInitialState(ERHIAccess::SRVMask);
+	
+		RenderableTexture = RHICreateTexture(RenderTargetTextureDesc);
+		ShaderTexture2D = RenderableTexture;
+	
+		CurrentResource->TextureRHI = RenderableTexture;
+
+		RHIUpdateTextureReference(TextureReference.TextureReferenceRHI, CurrentResource->TextureRHI);
+	});
+}
+
+void UPixelStreamingMediaTexture::UpdateTextureReference(FRHICommandList& RHICmdList, FTexture2DRHIRef Reference)
+{
+	if (CurrentResource != nullptr)
+	{
+		if (Reference.IsValid() && CurrentResource->TextureRHI != Reference)
+		{
+			CurrentResource->TextureRHI = Reference;
+			RHIUpdateTextureReference(TextureReference.TextureReferenceRHI, CurrentResource->TextureRHI);
+		}
+		else if (!Reference.IsValid())
+		{
+			if (CurrentResource != nullptr)
+			{
+				InitializeResources();
+
+				// Make sure RenderThread is executed before continuing
+				FlushRenderingCommands();
+			}
+		}
+	}
+}
diff --git a/Engine/Plugins/Experimental/PixelStreamingPlayer/Source/PixelStreamingPlayer/Private/PixelStreamingMediaTextureResource.cpp b/Engine/Plugins/Experimental/PixelStreamingPlayer/Source/PixelStreamingPlayer/Private/PixelStreamingMediaTextureResource.cpp
new file mode 100644
--- /dev/null
+++ b/Engine/Plugins/Experimental/PixelStreamingPlayer/Source/PixelStreamingPlayer/Private/PixelStreamingMediaTextureResource.cpp
@@ -0,0 +1,46 @@
+// Copyright Epic Games, Inc. All Rights Reserved.
+
+#include "PixelStreamingMediaTextureResource.h"
+#include "DeviceProfiles/DeviceProfile.h"
+#include "DeviceProfiles/DeviceProfileManager.h"
+
+FPixelStreamingMediaTextureResource::FPixelStreamingMediaTextureResource(UPixelStreamingMediaTexture* Owner)
+{
+	MediaTexture = Owner;
+}
+
+void FPixelStreamingMediaTextureResource::InitRHI(FRHICommandListBase& RHICmdList)
+{
+	if (MediaTexture != nullptr)
+	{
+		FSamplerStateInitializerRHI SamplerStateInitializer(
+			(ESamplerFilter)UDeviceProfileManager::Get().GetActiveProfile()->GetTextureLODSettings()->GetSamplerFilter(MediaTexture),
+			AM_Border, AM_Border, AM_Wrap);
+		SamplerStateRHI = RHICreateSamplerState(SamplerStateInitializer);
+	}
+}
+
+void FPixelStreamingMediaTextureResource::ReleaseRHI()
+{
+	TextureRHI.SafeRelease();
+
+	if (MediaTexture != nullptr)
+	{
+		RHIUpdateTextureReference(MediaTexture->TextureReference.TextureReferenceRHI, nullptr);
+	}
+}
+
+uint32 FPixelStreamingMediaTextureResource::GetSizeX() const
+{
+	return TextureRHI.IsValid() ? TextureRHI->GetSizeXYZ().X : 0;
+}
+
+uint32 FPixelStreamingMediaTextureResource::GetSizeY() const
+{
+	return TextureRHI.IsValid() ? TextureRHI->GetSizeXYZ().Y : 0;
+}
+
+SIZE_T FPixelStreamingMediaTextureResource::GetResourceSize()
+{
+	return CalcTextureSize(GetSizeX(), GetSizeY(), EPixelFormat::PF_A8R8G8B8, 1);
+}
diff --git a/Engine/Plugins/Experimental/PixelStreamingPlayer/Source/PixelStreamingPlayer/Private/PixelStreamingMediaTextureResource.h b/Engine/Plugins/Experimental/PixelStreamingPlayer/Source/PixelStreamingPlayer/Private/PixelStreamingMediaTextureResource.h
new file mode 100644
--- /dev/null
+++ b/Engine/Plugins/Experimental/PixelStreamingPlayer/Source/PixelStreamingPlayer/Private/PixelStreamingMediaTextureResource.h
@@ -0,0 +1,27 @@
+// Copyright Epic Games, Inc. All Rights Reserved.
+
+#pragma once
+
+#include "TextureResource.h"
+#include "PixelStreamingMediaTexture.h"
+
+/**
+ * The actual texture resource for a FPixelStreamingMediaTexture. Contains the RHI Texture and
+ * Sampler information.
+ */
+class PIXELSTREAMINGPLAYER_API FPixelStreamingMediaTextureResource : public FTextureResource
+{
+public:
+	FPixelStreamingMediaTextureResource(UPixelStreamingMediaTexture* Owner);
+	virtual ~FPixelStreamingMediaTextureResource() override { TextureRHI.SafeRelease(); }
+
+	virtual void InitRHI(FRHICommandListBase& RHICmdList) override;
+	virtual void ReleaseRHI() override;
+	virtual uint32 GetSizeX() const override;
+	virtual uint32 GetSizeY() const override;
+
+	SIZE_T GetResourceSize();
+
+private:
+	UPixelStreamingMediaTexture* MediaTexture = nullptr;
+};
diff --git a/Engine/Plugins/Experimental/PixelStreamingPlayer/Source/PixelStreamingPlayer/Private/PixelStreamingPeerComponent.cpp b/Engine/Plugins/Experimental/PixelStreamingPlayer/Source/PixelStreamingPlayer/Private/PixelStreamingPeerComponent.cpp
--- a/Engine/Plugins/Experimental/PixelStreamingPlayer/Source/PixelStreamingPlayer/Private/PixelStreamingPeerComponent.cpp
+++ b/Engine/Plugins/Experimental/PixelStreamingPlayer/Source/PixelStreamingPlayer/Private/PixelStreamingPeerComponent.cpp
@@ -3,10 +3,6 @@
 #include "PixelStreamingPeerComponent.h"
 #include "PixelStreamingPlayerPrivate.h"
 #include "PixelStreamingPeerConnection.h"
-#include "FrameBufferRHI.h"
-#include "FrameBufferI420.h"
-#include "ITextureMediaPlayer.h"
-#include "ITextureMediaPlayerModule.h"
 #include "RTCStatsCollector.h"
 #include "Async/Async.h"
 
@@ -26,23 +22,7 @@ void UPixelStreamingPeerComponent::SetConfig(const FPixelStreamingRTCConfigWrapp
 
 	if (PeerConnection)
 	{
-		if (VideoSinkPlayer != nullptr)
-		{
-			ITextureMediaPlayerModule* TextureModule = FModuleManager::LoadModulePtr<ITextureMediaPlayerModule>("TextureMediaPlayer");
-			if (TextureModule != nullptr)
-			{
-				TextureMediaPlayer = TextureModule->OpenPlayer(VideoSinkPlayer);
-				if (TextureMediaPlayer != nullptr)
-				{
-					PeerConnection->SetVideoSink(this);
-				}
-			}
-		}
-		else
-		{
-			UE_LOG(LogPixelStreamingPlayer, Error, TEXT("MediaPlayer not found."));
-		}
-
+		PeerConnection->SetVideoSink(VideoSink);
 		PeerConnection->OnEmitIceCandidate.AddLambda([this](const webrtc::IceCandidateInterface* Candidate) {
 			FPixelStreamingIceCandidateWrapper CandidateWrapper(*Candidate);
 			AsyncTask(ENamedThreads::GameThread, [this, CandidateWrapper]() {
@@ -141,7 +121,7 @@ FPixelStreamingSessionDescriptionWrapper UPixelStreamingPeerComponent::CreateAns
 	}
 }
 
-void UPixelStreamingPeerComponent::ReceiveAnswer(const FString& Offer)
+void UPixelStreamingPeerComponent::ReceiveAnswer(const FString& Answer)
 {
 	if (PeerConnection)
 	{
@@ -150,9 +130,9 @@ void UPixelStreamingPeerComponent::ReceiveAnswer(const FString& Offer)
 			UE_LOG(LogPixelStreamingPlayer, Error, TEXT("ReceiveAnswer Failed: %s"), *ErrorMsg);
 			TaskEvent->Trigger();
 		};
-		AsyncTask(ENamedThreads::AnyNormalThreadNormalTask, [this, &TaskEvent, &Offer, &OnGeneralFailure]() {
+		AsyncTask(ENamedThreads::AnyNormalThreadNormalTask, [this, &TaskEvent, &Answer, &OnGeneralFailure]() {
 			PeerConnection->ReceiveAnswer(
-				Offer,
+				Answer,
 				[this, &TaskEvent, &OnGeneralFailure]() {
 					// success
 					TaskEvent->Trigger();
@@ -187,44 +167,3 @@ void UPixelStreamingPeerComponent::OnIceConnectionChange(webrtc::PeerConnectionI
 		OnIceDisconnection.Broadcast();
 	}
 }
-
-void UPixelStreamingPeerComponent::OnFrame(const webrtc::VideoFrame& frame)
-{
-	if (frame.video_frame_buffer()->type() == webrtc::VideoFrameBuffer::Type::kNative)
-	{
-		UE::PixelStreaming::FFrameBufferRHI* const FrameBuffer = static_cast<UE::PixelStreaming::FFrameBufferRHI*>(frame.video_frame_buffer().get());
-
-		if (FrameBuffer != nullptr)
-		{
-			TSharedPtr<FVideoResourceRHI, ESPMode::ThreadSafe> VideoResource = static_cast<UE::PixelStreaming::FFrameBufferRHI*>(frame.video_frame_buffer().get())->GetVideoResource();
-
-			if (TextureMediaPlayer.IsValid() && VideoResource.IsValid())
-			{
-				ResourceQueue.Enqueue(VideoResource);
-
-				// Get off WebRTC thread ASAP
-				AsyncTask(ENamedThreads::AnyHiPriThreadHiPriTask, [this, TMPWeak = TextureMediaPlayer.ToWeakPtr()]() {
-					auto PinnedTextureMediaPlayer = TMPWeak.Pin();
-
-					if (PinnedTextureMediaPlayer.IsValid() && !ResourceQueue.IsEmpty())
-					{
-						TSharedPtr<FVideoResource, ESPMode::ThreadSafe> VideoResource;
-
-						while (!ResourceQueue.IsEmpty())
-						{
-							ResourceQueue.Dequeue(VideoResource);
-						}
-
-						auto& Raw = StaticCastSharedPtr<FVideoResourceRHI>(VideoResource)->GetRaw();
-						PinnedTextureMediaPlayer->OnFrame(Raw.Texture, Raw.Fence, Raw.FenceValue);
-					}
-				});
-			}
-		}
-	}
-	if (frame.video_frame_buffer()->type() == webrtc::VideoFrameBuffer::Type::kI420)
-	{
-		// TODO (aidan.possemiers)
-		unimplemented();
-	}
-}
diff --git a/Engine/Plugins/Experimental/PixelStreamingPlayer/Source/PixelStreamingPlayer/Private/WebRTCIncludes.h b/Engine/Plugins/Experimental/PixelStreamingPlayer/Source/PixelStreamingPlayer/Private/WebRTCIncludes.h
deleted file mode 100644
diff --git a/Engine/Plugins/Experimental/PixelStreamingPlayer/Source/PixelStreamingPlayer/Public/PixelStreamingMediaTexture.h b/Engine/Plugins/Experimental/PixelStreamingPlayer/Source/PixelStreamingPlayer/Public/PixelStreamingMediaTexture.h
new file mode 100644
--- /dev/null
+++ b/Engine/Plugins/Experimental/PixelStreamingPlayer/Source/PixelStreamingPlayer/Public/PixelStreamingMediaTexture.h
@@ -0,0 +1,44 @@
+// Copyright Epic Games, Inc. All Rights Reserved.
+
+#pragma once
+
+#include "Engine/Texture2DDynamic.h"
+#include "RenderTargetPool.h"
+#include "api/media_stream_interface.h"
+#include "PixelStreamingMediaTexture.generated.h"
+
+class FPixelStreamingMediaTextureResource;
+
+/**
+ * A Texture Object that can be used in materials etc. that takes updates from webrtc frames.
+ */
+UCLASS(NotBlueprintType, NotBlueprintable, HideDropdown, HideCategories = (ImportSettings, Compression, Texture, Adjustments, Compositing, LevelOfDetail, Object), META = (DisplayName = "PixelStreaming Media Texture"))
+class PIXELSTREAMINGPLAYER_API UPixelStreamingMediaTexture : public UTexture2DDynamic, public rtc::VideoSinkInterface<webrtc::VideoFrame>
+{
+	GENERATED_UCLASS_BODY()
+
+protected:
+	// UObject overrides.
+	virtual void BeginDestroy() override;
+	virtual void GetResourceSizeEx(FResourceSizeEx& CumulativeResourceSize) override;
+
+	// UTexture implementation
+	virtual FTextureResource* CreateResource() override;
+
+	// from rtc::VideoSinkInterface<webrtc::VideoFrame>
+	virtual void OnFrame(const webrtc::VideoFrame& frame) override;
+
+private:
+	void InitializeResources();
+
+	// updates the internal texture resource after each frame.
+	void UpdateTextureReference(FRHICommandList& RHICmdList, FTexture2DRHIRef Reference);
+
+	TArray<uint8_t> Buffer;
+
+	FCriticalSection RenderSyncContext;
+	FTextureRHIRef SourceTexture;
+	FPooledRenderTargetDesc RenderTargetDescriptor;
+	TRefCountPtr<IPooledRenderTarget> RenderTarget;
+	FPixelStreamingMediaTextureResource* CurrentResource;
+};
diff --git a/Engine/Plugins/Experimental/PixelStreamingPlayer/Source/PixelStreamingPlayer/Public/PixelStreamingPeerComponent.h b/Engine/Plugins/Experimental/PixelStreamingPlayer/Source/PixelStreamingPlayer/Public/PixelStreamingPeerComponent.h
--- a/Engine/Plugins/Experimental/PixelStreamingPlayer/Source/PixelStreamingPlayer/Public/PixelStreamingPeerComponent.h
+++ b/Engine/Plugins/Experimental/PixelStreamingPlayer/Source/PixelStreamingPlayer/Public/PixelStreamingPeerComponent.h
@@ -4,25 +4,20 @@
 
 #include "PixelStreamingPeerConnection.h"
 #include "PixelStreamingSignallingComponent.h"
+#include "PixelStreamingMediaTexture.h"
 #include "PixelStreamingWebRTCWrappers.h"
-#include "Containers/Queue.h"
 #include "PixelStreamingPeerComponent.generated.h"
 
 DECLARE_DYNAMIC_MULTICAST_SPARSE_DELEGATE_OneParam(FPixelStreamingOnIceCandidate, UPixelStreamingPeerComponent, OnIceCandidate, FPixelStreamingIceCandidateWrapper, Candidate);
 DECLARE_DYNAMIC_MULTICAST_SPARSE_DELEGATE_OneParam(FPixelStreamingOnIceConnection, UPixelStreamingPeerComponent, OnIceConnection, int, Number);
 DECLARE_DYNAMIC_MULTICAST_SPARSE_DELEGATE(FPixelStreamingOnIceDisconnection, UPixelStreamingPeerComponent, OnIceDisconnection);
 
-class ITextureMediaPlayer;
-class UMediaPlayer;
-class FVideoResource;
-
 /**
  * A blueprint representation of a Pixel Streaming Peer Connection. Should communicate with a Pixel Streaming Signalling Connection
  * and will accept video sinks to receive video data.
  */
 UCLASS(BlueprintType, Blueprintable, Category = "PixelStreaming", META = (DisplayName = "PixelStreaming Peer Component", BlueprintSpawnableComponent))
 class PIXELSTREAMINGPLAYER_API UPixelStreamingPeerComponent : public UActorComponent
-	, public rtc::VideoSinkInterface<webrtc::VideoFrame>
 {
 	GENERATED_UCLASS_BODY()
 
@@ -84,16 +79,11 @@ public:
 	/**
 	 * A sink for the video data received once this connection has finished negotiating.
 	 */
-	UPROPERTY(EditDefaultsOnly, Category = "Properties", META = (DisplayName = "Pixel Streaming Video Sink Player", AllowPrivateAccess = true))
-	TObjectPtr<UMediaPlayer> VideoSinkPlayer = nullptr;
-
-	// rtc::VideoSinkInterface<webrtc::VideoFrame> interface
-	virtual void OnFrame(const webrtc::VideoFrame& frame) override;
+	UPROPERTY(EditDefaultsOnly, Category = "Properties", META = (DisplayName = "Pixel Streaming Video Sink", AllowPrivateAccess = true))
+	TObjectPtr<UPixelStreamingMediaTexture> VideoSink = nullptr;
 
 private:
 	TUniquePtr<FPixelStreamingPeerConnection> PeerConnection;
-	TSharedPtr<ITextureMediaPlayer, ESPMode::ThreadSafe> TextureMediaPlayer;
-	TQueue<TSharedPtr<FVideoResource, ESPMode::ThreadSafe>> ResourceQueue;
 
 	void OnIceConnectionChange(webrtc::PeerConnectionInterface::IceConnectionState);
 };
diff --git a/Engine/Plugins/Experimental/PixelStreamingPlayer/Source/PixelStreamingPlayerEditor/Private/PixelStreamingMediaTextureFactory.cpp b/Engine/Plugins/Experimental/PixelStreamingPlayer/Source/PixelStreamingPlayerEditor/Private/PixelStreamingMediaTextureFactory.cpp
new file mode 100644
--- /dev/null
+++ b/Engine/Plugins/Experimental/PixelStreamingPlayer/Source/PixelStreamingPlayerEditor/Private/PixelStreamingMediaTextureFactory.cpp
@@ -0,0 +1,40 @@
+// Copyright Epic Games, Inc. All Rights Reserved.
+
+#include "PixelStreamingMediaTextureFactory.h"
+
+#include <AssetTypeCategories.h>
+#include "PixelStreamingMediaTexture.h"
+
+#define LOCTEXT_NAMESPACE "PixelStreaming"
+
+UPixelStreamingMediaTextureFactory::UPixelStreamingMediaTextureFactory(const FObjectInitializer& ObjectInitializer)
+	: Super(ObjectInitializer)
+{
+	bCreateNew = true;
+	bEditAfterNew = true;
+
+	SupportedClass = UPixelStreamingMediaTexture::StaticClass();
+}
+
+FText UPixelStreamingMediaTextureFactory::GetDisplayName() const
+{
+	return LOCTEXT("MediaTextureFactoryDisplayName", "Pixel Streaming Media Texture");
+}
+
+uint32 UPixelStreamingMediaTextureFactory::GetMenuCategories() const
+{
+	return EAssetTypeCategories::Textures;
+}
+
+UObject* UPixelStreamingMediaTextureFactory::FactoryCreateNew(UClass* InClass, UObject* InParent, FName InName, EObjectFlags Flags, UObject* Context, FFeedbackContext* Warn)
+{
+	if (UPixelStreamingMediaTexture* Resource = NewObject<UPixelStreamingMediaTexture>(InParent, InName, Flags | RF_Transactional))
+	{
+		Resource->UpdateResource();
+		return Resource;
+	}
+
+	return nullptr;
+}
+
+#undef LOCTEXT_NAMESPACE
diff --git a/Engine/Plugins/Experimental/PixelStreamingPlayer/Source/PixelStreamingPlayerEditor/Private/PixelStreamingMediaTextureFactory.h b/Engine/Plugins/Experimental/PixelStreamingPlayer/Source/PixelStreamingPlayerEditor/Private/PixelStreamingMediaTextureFactory.h
new file mode 100644
--- /dev/null
+++ b/Engine/Plugins/Experimental/PixelStreamingPlayer/Source/PixelStreamingPlayerEditor/Private/PixelStreamingMediaTextureFactory.h
@@ -0,0 +1,20 @@
+// Copyright Epic Games, Inc. All Rights Reserved.
+
+#pragma once
+
+#include <CoreMinimal.h>
+#include <Factories/Factory.h>
+#include <UObject/Object.h>
+
+#include "PixelStreamingMediaTextureFactory.generated.h"
+
+UCLASS()
+class UPixelStreamingMediaTextureFactory : public UFactory
+{
+	GENERATED_UCLASS_BODY()
+
+	public:
+		virtual FText GetDisplayName() const override;
+		virtual uint32 GetMenuCategories() const override;
+		virtual UObject* FactoryCreateNew(UClass* InClass, UObject* InParent, FName InName, EObjectFlags Flags, UObject* Context, FFeedbackContext* Warn) override;
+};
