Fix: Remove PixelStreamingPlayer's private include of PixelStreaming

diff --git a/Engine/Plugins/Experimental/PixelStreamingPlayer/Source/PixelStreamingPlayer/PixelStreamingPlayer.Build.cs b/Engine/Plugins/Experimental/PixelStreamingPlayer/Source/PixelStreamingPlayer/PixelStreamingPlayer.Build.cs
--- a/Engine/Plugins/Experimental/PixelStreamingPlayer/Source/PixelStreamingPlayer/PixelStreamingPlayer.Build.cs
+++ b/Engine/Plugins/Experimental/PixelStreamingPlayer/Source/PixelStreamingPlayer/PixelStreamingPlayer.Build.cs
@@ -17,11 +17,6 @@ namespace UnrealBuildTool.Rules
 				new string[] {
 				});
 
-			PrivateIncludePaths.AddRange(
-				new string[] {
-					Path.Combine(EngineDir, "Plugins/Media/PixelStreaming/Source/PixelStreaming/Private")
-				});
-
 			PublicDependencyModuleNames.AddRange(
 				new string[] {
 					"WebRTC",
diff --git a/Engine/Plugins/Experimental/PixelStreamingPlayer/Source/PixelStreamingPlayer/Private/PixelStreamingPeerComponent.cpp b/Engine/Plugins/Experimental/PixelStreamingPlayer/Source/PixelStreamingPlayer/Private/PixelStreamingPeerComponent.cpp
--- a/Engine/Plugins/Experimental/PixelStreamingPlayer/Source/PixelStreamingPlayer/Private/PixelStreamingPeerComponent.cpp
+++ b/Engine/Plugins/Experimental/PixelStreamingPlayer/Source/PixelStreamingPlayer/Private/PixelStreamingPeerComponent.cpp
@@ -3,7 +3,7 @@
 #include "PixelStreamingPeerComponent.h"
 #include "PixelStreamingPlayerPrivate.h"
 #include "PixelStreamingPeerConnection.h"
-#include "RTCStatsCollector.h"
+// #include "RTCStatsCollector.h"
 #include "Async/Async.h"
 
 UPixelStreamingPeerComponent::UPixelStreamingPeerComponent(const FObjectInitializer& ObjectInitializer)
@@ -14,11 +14,12 @@ UPixelStreamingPeerComponent::UPixelStreamingPeerComponent(const FObjectInitiali
 void UPixelStreamingPeerComponent::SetConfig(const FPixelStreamingRTCConfigWrapper& Config)
 {
 	PeerConnection = FPixelStreamingPeerConnection::Create(Config.Config);
-#if WEBRTC_5414
-	PeerConnection->SetWebRTCStatsCallback(rtc::scoped_refptr<UE::PixelStreaming::FRTCStatsCollector>(new UE::PixelStreaming::FRTCStatsCollector(TEXT("Streamer"))));
-#else
-	PeerConnection->SetWebRTCStatsCallback(new rtc::RefCountedObject<UE::PixelStreaming::FRTCStatsCollector>(TEXT("Streamer")));
-#endif
+	// TODO (william.belcher): Move the RTCStatsCollector from PixelStreaming to an Internal folder so it can be accessed across modules.
+	// #if WEBRTC_5414
+	// 	PeerConnection->SetWebRTCStatsCallback(rtc::scoped_refptr<UE::PixelStreaming::FRTCStatsCollector>(new UE::PixelStreaming::FRTCStatsCollector(TEXT("Streamer"))));
+	// #else
+	// 	PeerConnection->SetWebRTCStatsCallback(new rtc::RefCountedObject<UE::PixelStreaming::FRTCStatsCollector>(TEXT("Streamer")));
+	// #endif
 
 	if (PeerConnection)
 	{
@@ -46,10 +47,10 @@ FPixelStreamingSessionDescriptionWrapper UPixelStreamingPeerComponent::CreateOff
 	if (PeerConnection)
 	{
 		FPixelStreamingSessionDescriptionWrapper Wrapper;
-		FEvent* TaskEvent = FPlatformProcess::GetSynchEventFromPool();
-		const auto OnGeneralFailure = [&TaskEvent](const FString& ErrorMsg) {
-			UE_LOG(LogPixelStreamingPlayer, Error, TEXT("CreateOffer Failed: %s"), *ErrorMsg);
-			TaskEvent->Trigger();
+		FEvent*									 TaskEvent = FPlatformProcess::GetSynchEventFromPool();
+		const auto								 OnGeneralFailure = [&TaskEvent](const FString& ErrorMsg) {
+			  UE_LOG(LogPixelStreamingPlayer, Error, TEXT("CreateOffer Failed: %s"), *ErrorMsg);
+			  TaskEvent->Trigger();
 		};
 		AsyncTask(ENamedThreads::AnyNormalThreadNormalTask, [this, &TaskEvent, &Wrapper, &OnGeneralFailure]() {
 			PeerConnection->CreateOffer(
@@ -83,10 +84,10 @@ FPixelStreamingSessionDescriptionWrapper UPixelStreamingPeerComponent::CreateAns
 	if (PeerConnection)
 	{
 		FPixelStreamingSessionDescriptionWrapper Wrapper;
-		FEvent* TaskEvent = FPlatformProcess::GetSynchEventFromPool();
-		const auto OnGeneralFailure = [&TaskEvent](const FString& ErrorMsg) {
-			UE_LOG(LogPixelStreamingPlayer, Error, TEXT("CreateAnswer Failed: %s"), *ErrorMsg);
-			TaskEvent->Trigger();
+		FEvent*									 TaskEvent = FPlatformProcess::GetSynchEventFromPool();
+		const auto								 OnGeneralFailure = [&TaskEvent](const FString& ErrorMsg) {
+			  UE_LOG(LogPixelStreamingPlayer, Error, TEXT("CreateAnswer Failed: %s"), *ErrorMsg);
+			  TaskEvent->Trigger();
 		};
 		AsyncTask(ENamedThreads::AnyNormalThreadNormalTask, [this, &TaskEvent, &Sdp, &Wrapper, &OnGeneralFailure]() {
 			PeerConnection->ReceiveOffer(
@@ -125,7 +126,7 @@ void UPixelStreamingPeerComponent::ReceiveAnswer(const FString& Answer)
 {
 	if (PeerConnection)
 	{
-		FEvent* TaskEvent = FPlatformProcess::GetSynchEventFromPool();
+		FEvent*	   TaskEvent = FPlatformProcess::GetSynchEventFromPool();
 		const auto OnGeneralFailure = [&TaskEvent](const FString& ErrorMsg) {
 			UE_LOG(LogPixelStreamingPlayer, Error, TEXT("ReceiveAnswer Failed: %s"), *ErrorMsg);
 			TaskEvent->Trigger();
diff --git a/Engine/Plugins/Experimental/PixelStreamingPlayer/Source/PixelStreamingPlayer/Private/WebRTCIncludes.h b/Engine/Plugins/Experimental/PixelStreamingPlayer/Source/PixelStreamingPlayer/Private/WebRTCIncludes.h
new file mode 100644
--- /dev/null
+++ b/Engine/Plugins/Experimental/PixelStreamingPlayer/Source/PixelStreamingPlayer/Private/WebRTCIncludes.h
@@ -0,0 +1,29 @@
+// Copyright Epic Games, Inc. All Rights Reserved.
+
+#pragma once
+
+THIRD_PARTY_INCLUDES_START
+
+#if PLATFORM_WINDOWS
+
+	#include "Windows/AllowWindowsPlatformTypes.h"
+
+	// C4582: constructor is not implicitly called in "api/rtcerror.h", treated as an error by UnrealEngine
+	// C6319: Use of the comma-operator in a tested expression causes the left argument to be ignored when it has no side-effects.
+	// C6323: Use of arithmetic operator on Boolean type(s).
+	#pragma warning(push)
+	#pragma warning(disable : 6319 6323)
+#endif // PLATFORM_WINDOWS
+
+#include "rtc_base/ssl_adapter.h"
+#include "modules/video_capture/video_capture_defines.h"
+
+#if PLATFORM_WINDOWS
+	#pragma warning(pop)
+
+	#include "Windows/HideWindowsPlatformTypes.h"
+#endif // PLATFORM_WINDOWS
+
+THIRD_PARTY_INCLUDES_END
+
+#include "PixelStreamingPlayerPrivate.h"
\ No newline at end of file
