QOL: Default Pixel Streaming Player connection URL to port 80

diff --git a/Engine/Plugins/Experimental/PixelStreamingPlayer/Source/PixelStreamingPlayer/Private/PixelStreamingSignallingComponent.cpp b/Engine/Plugins/Experimental/PixelStreamingPlayer/Source/PixelStreamingPlayer/Private/PixelStreamingSignallingComponent.cpp
--- a/Engine/Plugins/Experimental/PixelStreamingPlayer/Source/PixelStreamingPlayer/Private/PixelStreamingSignallingComponent.cpp
+++ b/Engine/Plugins/Experimental/PixelStreamingPlayer/Source/PixelStreamingPlayer/Private/PixelStreamingSignallingComponent.cpp
@@ -22,14 +22,24 @@ void UPixelStreamingSignallingComponent::Connect(const FString& Url)
 {
 	SignallingConnection->SetKeepAlive(false);
 	SignallingConnection->SetAutoReconnect(true);
-	if (MediaSource == nullptr)
+	FString URL = (MediaSource == nullptr ? Url : MediaSource->GetUrl());
+
+	TArray<FString> Components;
+	URL.ParseIntoArray(Components, TEXT(":"), true);
+
+	if(Components.Num() < 2)
 	{
-		SignallingConnection->TryConnect(Url);
+		UE_LOG(LogPixelStreamingPlayer, Error, TEXT("Incorrectly formated connection URL. Ensure URL is in the format (protocol)://(ip):(port)"));
+		return;
 	}
-	else
+
+	if(Components.Num() == 2)
 	{
-		SignallingConnection->TryConnect(MediaSource->GetUrl());
+		// The user hasn't specified the port. Append port 80 as Pixel Streaming defaults to 8888
+		URL += TEXT(":80");
 	}
+	
+	SignallingConnection->TryConnect(URL);
 }
 
 void UPixelStreamingSignallingComponent::Disconnect()
