[PixelStreaming] Fix: Undeprecate as VCam is still depending on it

diff --git a/Engine/Plugins/Experimental/PixelStreamingPlayer/PixelStreamingPlayer.uplugin b/Engine/Plugins/Experimental/PixelStreamingPlayer/PixelStreamingPlayer.uplugin
--- a/Engine/Plugins/Experimental/PixelStreamingPlayer/PixelStreamingPlayer.uplugin
+++ b/Engine/Plugins/Experimental/PixelStreamingPlayer/PixelStreamingPlayer.uplugin
@@ -14,7 +14,6 @@
 	"CanContainContent": true,
 	"IsBetaVersion": true,
 	"Installed": false,
-	"DeprecatedEngineVersion": "5.5",
 	"Modules": [
 		{
 			"Name": "PixelStreamingPlayer",
diff --git a/Engine/Plugins/Media/PixelStreaming/PixelStreaming.uplugin b/Engine/Plugins/Media/PixelStreaming/PixelStreaming.uplugin
--- a/Engine/Plugins/Media/PixelStreaming/PixelStreaming.uplugin
+++ b/Engine/Plugins/Media/PixelStreaming/PixelStreaming.uplugin
@@ -14,7 +14,6 @@
 	"CanContainContent": true,
 	"IsBetaVersion": false,
 	"Installed": false,
-	"DeprecatedEngineVersion": "5.5",
 	"Modules": [
 		{
 			"Name": "PixelStreaming",
